UE4 - fix nullrhi to support very large textures

diff --git a/Engine/Source/Runtime/NullDrv/Private/NullRHI.cpp b/Engine/Source/Runtime/NullDrv/Private/NullRHI.cpp
--- a/Engine/Source/Runtime/NullDrv/Private/NullRHI.cpp
+++ b/Engine/Source/Runtime/NullDrv/Private/NullRHI.cpp
@@ -53,8 +53,8 @@ void* FNullDynamicRHI::GetStaticBuffer()
 	static void* Buffer = NULL;
 	if (!Buffer)
 	{
-		// allocate an 16 meg buffer, should be big enough for any texture/surface
-		Buffer = FMemory::Malloc(16 * 1024 * 1024);
+		// allocate an 64 meg buffer, should be big enough for any texture/surface
+		Buffer = FMemory::Malloc(64 * 1024 * 1024);
 	}
 
 	return Buffer;
