Fix for crash when cooking on mac
 - In NullRHI, initialise featurelevel->shaderplatform mapping table to opengl platforms when not in windows

diff --git a/Engine/Source/Runtime/NullDrv/Private/NullRHI.cpp b/Engine/Source/Runtime/NullDrv/Private/NullRHI.cpp
--- a/Engine/Source/Runtime/NullDrv/Private/NullRHI.cpp
+++ b/Engine/Source/Runtime/NullDrv/Private/NullRHI.cpp
@@ -19,11 +19,18 @@ FNullDynamicRHI::FNullDynamicRHI()
 
 void FNullDynamicRHI::Init()
 {
-	GShaderPlatformForFeatureLevel[ERHIFeatureLevel::ES2] = SP_PCD3D_ES2;
+#if PLATFORM_WINDOWS
+    GShaderPlatformForFeatureLevel[ERHIFeatureLevel::ES2] = SP_PCD3D_ES2;
 	GShaderPlatformForFeatureLevel[ERHIFeatureLevel::SM3] = SP_NumPlatforms;
 	GShaderPlatformForFeatureLevel[ERHIFeatureLevel::SM4] = SP_PCD3D_SM4;
 	GShaderPlatformForFeatureLevel[ERHIFeatureLevel::SM5] = SP_PCD3D_SM5;
-
+#else
+    GShaderPlatformForFeatureLevel[ERHIFeatureLevel::ES2] = SP_OPENGL_ES2;
+	GShaderPlatformForFeatureLevel[ERHIFeatureLevel::SM3] = SP_NumPlatforms;
+	GShaderPlatformForFeatureLevel[ERHIFeatureLevel::SM4] = SP_OPENGL_SM4;
+	GShaderPlatformForFeatureLevel[ERHIFeatureLevel::SM5] = SP_OPENGL_SM5;
+#endif
+    
 	check(!GIsRHIInitialized);
 
 	// Notify all initialized FRenderResources that there's a valid RHI device to create their RHI resources for now.
