NullRHI now initialises to the correct Metal shader platforms for Mac as we no longer support OpenGL on Mac at all.

diff --git a/Engine/Source/Runtime/NullDrv/Private/NullRHI.cpp b/Engine/Source/Runtime/NullDrv/Private/NullRHI.cpp
--- a/Engine/Source/Runtime/NullDrv/Private/NullRHI.cpp
+++ b/Engine/Source/Runtime/NullDrv/Private/NullRHI.cpp
@@ -22,6 +22,11 @@ void FNullDynamicRHI::Init()
 	GShaderPlatformForFeatureLevel[ERHIFeatureLevel::ES3_1] = SP_PCD3D_ES3_1;
 	GShaderPlatformForFeatureLevel[ERHIFeatureLevel::SM4] = SP_PCD3D_SM4;
 	GShaderPlatformForFeatureLevel[ERHIFeatureLevel::SM5] = SP_PCD3D_SM5;
+#elif PLATFORM_MAC
+    GShaderPlatformForFeatureLevel[ERHIFeatureLevel::ES2] = SP_METAL_MACES2;
+    GShaderPlatformForFeatureLevel[ERHIFeatureLevel::ES3_1] = SP_METAL_MACES3_1;
+    GShaderPlatformForFeatureLevel[ERHIFeatureLevel::SM4] = SP_METAL_SM4;
+    GShaderPlatformForFeatureLevel[ERHIFeatureLevel::SM5] = SP_METAL_SM5;
 #else
 	GShaderPlatformForFeatureLevel[ERHIFeatureLevel::ES2] = SP_OPENGL_PCES2;
 	GShaderPlatformForFeatureLevel[ERHIFeatureLevel::ES3_1] = SP_NumPlatforms;
