UE4.23 - Deprecate FUniformBufferRHIParamRef, FIndexBufferRHIParamRef, FVertexBufferRHIParamRef, FStructuredBufferRHIParamRef, FTextureRHIParamRef, FTexture2DRHIParamRef, FTexture2DArrayRHIParamRef, F

diff --git a/Engine/Source/Runtime/Engine/Private/VT/VirtualTextureUploadCache.cpp b/Engine/Source/Runtime/Engine/Private/VT/VirtualTextureUploadCache.cpp
--- a/Engine/Source/Runtime/Engine/Private/VT/VirtualTextureUploadCache.cpp
+++ b/Engine/Source/Runtime/Engine/Private/VT/VirtualTextureUploadCache.cpp
@@ -232,7 +232,7 @@ FVTUploadTileHandle FVirtualTextureUploadCache::PrepareTileForUpload(FVTUploadTi
 	return FVTUploadTileHandle(Index);
 }
 
-void FVirtualTextureUploadCache::SubmitTile(FRHICommandListImmediate& RHICmdList, const FVTUploadTileHandle& InHandle, FTexture2DRHIParamRef InDestTexture, int InDestX, int InDestY, int InSkipBorderSize)
+void FVirtualTextureUploadCache::SubmitTile(FRHICommandListImmediate& RHICmdList, const FVTUploadTileHandle& InHandle, FRHITexture2D* InDestTexture, int InDestX, int InDestY, int InSkipBorderSize)
 {
 	checkSlow(IsInRenderingThread());
 	check(NumPendingTiles > 0u);
diff --git a/Engine/Source/Runtime/Engine/Private/VectorField.cpp b/Engine/Source/Runtime/Engine/Private/VectorField.cpp
--- a/Engine/Source/Runtime/Engine/Private/VectorField.cpp
+++ b/Engine/Source/Runtime/Engine/Private/VectorField.cpp
@@ -808,8 +808,8 @@ public:
 	void SetParameters(
 		FRHICommandList& RHICmdList, 
 		const FCompositeAnimatedVectorFieldUniformBufferRef& UniformBuffer,
-		FTextureRHIParamRef AtlasTextureRHI,
-		FTextureRHIParamRef NoiseVolumeTextureRHI )
+		FRHITexture* AtlasTextureRHI,
+		FRHITexture* NoiseVolumeTextureRHI )
 	{
 		FRHIComputeShader* ComputeShaderRHI = GetComputeShader();
 		FRHISamplerState* SamplerStateLinear = TStaticSamplerState<SF_Bilinear,AM_Clamp,AM_Clamp,AM_Clamp>::GetRHI();
@@ -988,7 +988,7 @@ public:
 				FCompositeAnimatedVectorFieldUniformBufferRef::CreateUniformBufferImmediate(Parameters, UniformBuffer_SingleDraw);
 
 			TShaderMapRef<FCompositeAnimatedVectorFieldCS> CompositeCS(GetGlobalShaderMap(GetFeatureLevel()));
-			FTextureRHIParamRef NoiseVolumeTextureRHI = GBlackVolumeTexture->TextureRHI;
+			FRHITexture* NoiseVolumeTextureRHI = GBlackVolumeTexture->TextureRHI;
 			if (AnimatedVectorField->NoiseField && AnimatedVectorField->NoiseField->Resource)
 			{
 				NoiseVolumeTextureRHI = AnimatedVectorField->NoiseField->Resource->VolumeTextureRHI;
diff --git a/Engine/Source/Runtime/Engine/Private/VectorFieldVisualization.cpp b/Engine/Source/Runtime/Engine/Private/VectorFieldVisualization.cpp
--- a/Engine/Source/Runtime/Engine/Private/VectorFieldVisualization.cpp
+++ b/Engine/Source/Runtime/Engine/Private/VectorFieldVisualization.cpp
@@ -131,7 +131,7 @@ void FVectorFieldVisualizationVertexFactory::InitRHI()
 void FVectorFieldVisualizationVertexFactory::ReleaseRHI()
 {
 	UniformBuffer.SafeRelease();
-	VectorFieldTextureRHI = FTexture3DRHIParamRef();
+	VectorFieldTextureRHI = nullptr;
 	FVertexFactory::ReleaseRHI();
 }
 
@@ -165,7 +165,7 @@ FVertexFactoryShaderParameters* FVectorFieldVisualizationVertexFactory::Construc
  */
 void FVectorFieldVisualizationVertexFactory::SetParameters(
 	const FVectorFieldVisualizationParameters& InUniformParameters,
-	FTexture3DRHIParamRef InVectorFieldTextureRHI )
+	FRHITexture3D* InVectorFieldTextureRHI )
 {
 	UniformBuffer = FVectorFieldVisualizationBufferRef::CreateUniformBufferImmediate(InUniformParameters, UniformBuffer_SingleFrame);
 	VectorFieldTextureRHI = InVectorFieldTextureRHI;
diff --git a/Engine/Source/Runtime/Engine/Private/VectorFieldVisualization.h b/Engine/Source/Runtime/Engine/Private/VectorFieldVisualization.h
--- a/Engine/Source/Runtime/Engine/Private/VectorFieldVisualization.h
+++ b/Engine/Source/Runtime/Engine/Private/VectorFieldVisualization.h
@@ -87,7 +87,7 @@ public:
 	 */
 	void SetParameters(
 		const FVectorFieldVisualizationParameters& InUniformParameters,
-		FTexture3DRHIParamRef InVectorFieldTextureRHI );
+		FRHITexture3D* InVectorFieldTextureRHI );
 
 private:
 
@@ -97,7 +97,7 @@ private:
 	/** Uniform buffer. */
 	FUniformBufferRHIRef UniformBuffer;
 	/** Texture containing the vector field. */
-	FTexture3DRHIParamRef VectorFieldTextureRHI;
+	FRHITexture3D* VectorFieldTextureRHI;
 };
 
 /*------------------------------------------------------------------------------
diff --git a/Engine/Source/Runtime/Engine/Private/VolumeTexture.cpp b/Engine/Source/Runtime/Engine/Private/VolumeTexture.cpp
--- a/Engine/Source/Runtime/Engine/Private/VolumeTexture.cpp
+++ b/Engine/Source/Runtime/Engine/Private/VolumeTexture.cpp
@@ -480,7 +480,7 @@ public:
 		DEC_DWORD_STAT_FNAME_BY( LODGroupStatName, TextureSize );
 		if (TextureReference)
 		{
-			RHIUpdateTextureReference(TextureReference->TextureReferenceRHI, FTextureRHIParamRef());
+			RHIUpdateTextureReference(TextureReference->TextureReferenceRHI, nullptr);
 		}
 		Texture3DRHI.SafeRelease();
 		FTextureResource::ReleaseRHI();
diff --git a/Engine/Source/Runtime/Engine/Public/FXSystem.h b/Engine/Source/Runtime/Engine/Public/FXSystem.h
--- a/Engine/Source/Runtime/Engine/Public/FXSystem.h
+++ b/Engine/Source/Runtime/Engine/Public/FXSystem.h
@@ -212,14 +212,14 @@ public:
 	 */
 	virtual void PostRenderOpaque(
 		FRHICommandListImmediate& RHICmdList, 
-		const FUniformBufferRHIParamRef ViewUniformBuffer, 
+		FRHIUniformBuffer* ViewUniformBuffer,
 		const class FShaderParametersMetadata* SceneTexturesUniformBufferStruct,
-		FUniformBufferRHIParamRef SceneTexturesUniformBuffer) = 0;
+		FRHIUniformBuffer* SceneTexturesUniformBuffer) = 0;
 
 	/**
 	 * Helper in case the data necessary for collision is not available.
 	 */
-	void PostRenderOpaque(FRHICommandListImmediate& RHICmdList) { PostRenderOpaque(RHICmdList, nullptr, nullptr, FUniformBufferRHIParamRef()); }
+	void PostRenderOpaque(FRHICommandListImmediate& RHICmdList) { PostRenderOpaque(RHICmdList, nullptr, nullptr, nullptr); }
 
 	bool IsPendingKill() const { return bIsPendingKill; }
 
diff --git a/Engine/Source/Runtime/Engine/Public/GPUSkinVertexFactory.h b/Engine/Source/Runtime/Engine/Public/GPUSkinVertexFactory.h
--- a/Engine/Source/Runtime/Engine/Public/GPUSkinVertexFactory.h
+++ b/Engine/Source/Runtime/Engine/Public/GPUSkinVertexFactory.h
@@ -187,7 +187,7 @@ public:
 		}
 		
 		// if FeatureLevel <= ERHIFeatureLevel::ES3_1
-		FUniformBufferRHIParamRef GetUniformBuffer() const
+		FRHIUniformBuffer* GetUniformBuffer() const
 		{
 			return UniformBuffer;
 		}
diff --git a/Engine/Source/Runtime/Engine/Public/LocalVertexFactory.h b/Engine/Source/Runtime/Engine/Public/LocalVertexFactory.h
--- a/Engine/Source/Runtime/Engine/Public/LocalVertexFactory.h
+++ b/Engine/Source/Runtime/Engine/Public/LocalVertexFactory.h
@@ -131,7 +131,7 @@ public:
 		return Data.NumTexCoords;
 	}
 
-	FUniformBufferRHIParamRef GetUniformBuffer() const
+	FRHIUniformBuffer* GetUniformBuffer() const
 	{
 		return UniformBuffer.GetReference();
 	}
@@ -175,7 +175,7 @@ public:
 		ERHIFeatureLevel::Type FeatureLevel,
 		const FVertexFactory* VertexFactory,
 		const FMeshBatchElement& BatchElement,
-		FUniformBufferRHIParamRef VertexFactoryUniformBuffer,
+		FRHIUniformBuffer* VertexFactoryUniformBuffer,
 		FMeshDrawSingleShaderBindings& ShaderBindings,
 		FVertexInputStreamArray& VertexStreams
 		) const;
diff --git a/Engine/Source/Runtime/Engine/Public/MeshBatch.h b/Engine/Source/Runtime/Engine/Public/MeshBatch.h
--- a/Engine/Source/Runtime/Engine/Public/MeshBatch.h
+++ b/Engine/Source/Runtime/Engine/Public/MeshBatch.h
@@ -45,11 +45,11 @@ struct FMeshBatchElement
 	 * Primitive uniform buffer RHI
 	 * Must be null for vertex factories that manually fetch primitive data from scene data, in which case FPrimitiveSceneProxy::UniformBuffer will be used.
 	 */
-	FUniformBufferRHIParamRef PrimitiveUniformBuffer;
+	FRHIUniformBuffer* PrimitiveUniformBuffer;
 
 	/** 
 	 * Primitive uniform buffer to use for rendering, used when PrimitiveUniformBuffer is null. 
-	 * This interface allows a FMeshBatchElement to be setup for a uniform buffer that has not been initialized yet, (TUniformBuffer* is known but not the FUniformBufferRHIParamRef)
+	 * This interface allows a FMeshBatchElement to be setup for a uniform buffer that has not been initialized yet, (TUniformBuffer* is known but not the FRHIUniformBuffer*)
 	 */
 	const TUniformBuffer<FPrimitiveUniformShaderParameters>* PrimitiveUniformBufferResource;
 
@@ -96,7 +96,7 @@ struct FMeshBatchElement
 	/** Conceptual element index used for debug viewmodes. */
 	int8 VisualizeElementIndex;
 #endif
-	FVertexBufferRHIParamRef IndirectArgsBuffer;
+	FRHIVertexBuffer* IndirectArgsBuffer;
 	uint32 IndirectArgsOffset;
 
 	FMeshBatchElement()
@@ -313,7 +313,7 @@ struct FMeshBatch
 struct FUniformBufferValue
 {
 	const FShaderParametersMetadata* Type = nullptr;
-	FUniformBufferRHIParamRef UniformBuffer;
+	FRHIUniformBuffer* UniformBuffer;
 };
 
 
diff --git a/Engine/Source/Runtime/Engine/Public/MeshParticleVertexFactory.h b/Engine/Source/Runtime/Engine/Public/MeshParticleVertexFactory.h
--- a/Engine/Source/Runtime/Engine/Public/MeshParticleVertexFactory.h
+++ b/Engine/Source/Runtime/Engine/Public/MeshParticleVertexFactory.h
@@ -139,7 +139,7 @@ public:
 	/**
 	 * Retrieve the uniform buffer for this vertex factory.
 	 */
-	FORCEINLINE FUniformBufferRHIParamRef GetUniformBuffer()
+	FORCEINLINE FRHIUniformBuffer* GetUniformBuffer()
 	{
 		return MeshParticleUniformBuffer;
 	}
@@ -206,7 +206,7 @@ protected:
 	int32 DynamicParameterVertexStride;
 	
 	/** Uniform buffer with mesh particle parameters. */
-	FUniformBufferRHIParamRef MeshParticleUniformBuffer;
+	FRHIUniformBuffer* MeshParticleUniformBuffer;
 
 	FDynamicReadBuffer PrevTransformBuffer;
 
diff --git a/Engine/Source/Runtime/Engine/Public/ParameterCollection.h b/Engine/Source/Runtime/Engine/Public/ParameterCollection.h
--- a/Engine/Source/Runtime/Engine/Public/ParameterCollection.h
+++ b/Engine/Source/Runtime/Engine/Public/ParameterCollection.h
@@ -45,7 +45,7 @@ public:
 		return OwnerName;
 	}
 
-	FUniformBufferRHIParamRef GetUniformBuffer() const
+	FRHIUniformBuffer* GetUniformBuffer() const
 	{
 		return UniformBuffer;
 	}
diff --git a/Engine/Source/Runtime/Engine/Public/ParticleHelper.h b/Engine/Source/Runtime/Engine/Public/ParticleHelper.h
--- a/Engine/Source/Runtime/Engine/Public/ParticleHelper.h
+++ b/Engine/Source/Runtime/Engine/Public/ParticleHelper.h
@@ -2576,7 +2576,7 @@ public:
 	void SetVisualizeLODIndex(int32 InVisualizeLODIndex) { VisualizeLODIndex = InVisualizeLODIndex; }
 	int32  GetVisualizeLODIndex() const { return VisualizeLODIndex; }
 
-	inline FUniformBufferRHIParamRef GetWorldSpacePrimitiveUniformBuffer() const { return WorldSpacePrimitiveUniformBuffer.GetUniformBufferRHI(); }
+	inline FRHIUniformBuffer* GetWorldSpacePrimitiveUniformBuffer() const { return WorldSpacePrimitiveUniformBuffer.GetUniformBufferRHI(); }
 
 	const FColoredMaterialRenderProxy* GetDeselectedWireframeMatInst() const	{	return &DeselectedWireframeMaterialInstance;	}
 
diff --git a/Engine/Source/Runtime/Engine/Public/ParticleResources.h b/Engine/Source/Runtime/Engine/Public/ParticleResources.h
--- a/Engine/Source/Runtime/Engine/Public/ParticleResources.h
+++ b/Engine/Source/Runtime/Engine/Public/ParticleResources.h
@@ -57,7 +57,7 @@ ENGINE_API extern TGlobalResource<FParticleIndexBuffer> GParticleIndexBuffer;
 ENGINE_API extern TGlobalResource<FSixTriangleParticleIndexBuffer> GSixTriangleParticleIndexBuffer;
 
 typedef FRHIShaderResourceView* FParticleShaderParamRef;
-typedef FVertexBufferRHIParamRef FParticleBufferParamRef;
+typedef FRHIVertexBuffer* FParticleBufferParamRef;
 
 /**
  * Scratch vertex buffer available for dynamic draw calls.
diff --git a/Engine/Source/Runtime/Engine/Public/ParticleVertexFactory.h b/Engine/Source/Runtime/Engine/Public/ParticleVertexFactory.h
--- a/Engine/Source/Runtime/Engine/Public/ParticleVertexFactory.h
+++ b/Engine/Source/Runtime/Engine/Public/ParticleVertexFactory.h
@@ -209,7 +209,7 @@ public:
 	/**
 	 * Retrieve the uniform buffer for this vertex factory.
 	 */
-	FORCEINLINE FUniformBufferRHIParamRef GetSpriteUniformBuffer()
+	FORCEINLINE FRHIUniformBuffer* GetSpriteUniformBuffer()
 	{
 		return SpriteUniformBuffer;
 	}
diff --git a/Engine/Source/Runtime/Engine/Public/PrimitiveSceneProxy.h b/Engine/Source/Runtime/Engine/Public/PrimitiveSceneProxy.h
--- a/Engine/Source/Runtime/Engine/Public/PrimitiveSceneProxy.h
+++ b/Engine/Source/Runtime/Engine/Public/PrimitiveSceneProxy.h
@@ -517,7 +517,7 @@ public:
 	inline bool DoesVFRequirePrimitiveUniformBuffer() const { return bVFRequiresPrimitiveUniformBuffer; }
 	inline bool ShouldUseAsOccluder() const { return bUseAsOccluder; }
 	inline bool AllowApproximateOcclusion() const { return bAllowApproximateOcclusion; }
-	inline FUniformBufferRHIParamRef GetUniformBuffer() const 
+	inline FRHIUniformBuffer* GetUniformBuffer() const
 	{
 		return UniformBuffer.GetReference(); 
 	}
diff --git a/Engine/Source/Runtime/Engine/Public/Rendering/ColorVertexBuffer.h b/Engine/Source/Runtime/Engine/Public/Rendering/ColorVertexBuffer.h
--- a/Engine/Source/Runtime/Engine/Public/Rendering/ColorVertexBuffer.h
+++ b/Engine/Source/Runtime/Engine/Public/Rendering/ColorVertexBuffer.h
@@ -142,7 +142,7 @@ public:
 
 	/** Similar to Init/ReleaseRHI but only update existing SRV so references to the SRV stays valid */
 	template <uint32 MaxNumUpdates>
-	void InitRHIForStreaming(FVertexBufferRHIParamRef IntermediateBuffer, TRHIResourceUpdateBatcher<MaxNumUpdates>& Batcher)
+	void InitRHIForStreaming(FRHIVertexBuffer* IntermediateBuffer, TRHIResourceUpdateBatcher<MaxNumUpdates>& Batcher)
 	{
 		if (VertexBufferRHI && IntermediateBuffer)
 		{
diff --git a/Engine/Source/Runtime/Engine/Public/Rendering/PositionVertexBuffer.h b/Engine/Source/Runtime/Engine/Public/Rendering/PositionVertexBuffer.h
--- a/Engine/Source/Runtime/Engine/Public/Rendering/PositionVertexBuffer.h
+++ b/Engine/Source/Runtime/Engine/Public/Rendering/PositionVertexBuffer.h
@@ -100,7 +100,7 @@ public:
 
 	/** Similar to Init/ReleaseRHI but only update existing SRV so references to the SRV stays valid */
 	template <uint32 MaxNumUpdates>
-	void InitRHIForStreaming(FVertexBufferRHIParamRef IntermediateBuffer, TRHIResourceUpdateBatcher<MaxNumUpdates>& Batcher)
+	void InitRHIForStreaming(FRHIVertexBuffer* IntermediateBuffer, TRHIResourceUpdateBatcher<MaxNumUpdates>& Batcher)
 	{
 		check(VertexBufferRHI);
 		if (IntermediateBuffer)
diff --git a/Engine/Source/Runtime/Engine/Public/Rendering/StaticMeshVertexBuffer.h b/Engine/Source/Runtime/Engine/Public/Rendering/StaticMeshVertexBuffer.h
--- a/Engine/Source/Runtime/Engine/Public/Rendering/StaticMeshVertexBuffer.h
+++ b/Engine/Source/Runtime/Engine/Public/Rendering/StaticMeshVertexBuffer.h
@@ -397,8 +397,8 @@ public:
 	/** Similar to Init/ReleaseRHI but only update existing SRV so references to the SRV stays valid */
 	template <uint32 MaxNumUpdates>
 	void InitRHIForStreaming(
-		FVertexBufferRHIParamRef IntermediateTangentsBuffer,
-		FVertexBufferRHIParamRef IntermediateTexCoordBuffer,
+		FRHIVertexBuffer* IntermediateTangentsBuffer,
+		FRHIVertexBuffer* IntermediateTexCoordBuffer,
 		TRHIResourceUpdateBatcher<MaxNumUpdates>& Batcher)
 	{
 		check(TangentsVertexBuffer.VertexBufferRHI && TexCoordVertexBuffer.VertexBufferRHI);
diff --git a/Engine/Source/Runtime/Engine/Public/SceneInterface.h b/Engine/Source/Runtime/Engine/Public/SceneInterface.h
--- a/Engine/Source/Runtime/Engine/Public/SceneInterface.h
+++ b/Engine/Source/Runtime/Engine/Public/SceneInterface.h
@@ -284,7 +284,7 @@ public:
 	 * Looks up the SpeedTree uniform buffer for the passed in vertex factory.
 	 * @param VertexFactory - The vertex factory registered for SpeedTree.
 	 */
-	virtual FUniformBufferRHIParamRef GetSpeedTreeUniformBuffer(const FVertexFactory* VertexFactory) const = 0;
+	virtual FRHIUniformBuffer* GetSpeedTreeUniformBuffer(const FVertexFactory* VertexFactory) const = 0;
 
 	/**
 	 * Release this scene and remove it from the rendering thread
diff --git a/Engine/Source/Runtime/Engine/Public/SceneManagement.h b/Engine/Source/Runtime/Engine/Public/SceneManagement.h
--- a/Engine/Source/Runtime/Engine/Public/SceneManagement.h
+++ b/Engine/Source/Runtime/Engine/Public/SceneManagement.h
@@ -802,12 +802,12 @@ public:
 		bGlobalVolumeLightmap = bInGlobalVolumeLightmap;
 	}
 
-	FUniformBufferRHIParamRef GetPrecomputedLightingBuffer() const
+	FRHIUniformBuffer* GetPrecomputedLightingBuffer() const
 	{
 		return PrecomputedLightingUniformBuffer;
 	}
 
-	void SetPrecomputedLightingBuffer(FUniformBufferRHIParamRef InPrecomputedLightingUniformBuffer)
+	void SetPrecomputedLightingBuffer(FRHIUniformBuffer* InPrecomputedLightingUniformBuffer)
 	{
 		PrecomputedLightingUniformBuffer = InPrecomputedLightingUniformBuffer;
 	}
diff --git a/Engine/Source/Runtime/Engine/Public/ScreenRendering.h b/Engine/Source/Runtime/Engine/Public/ScreenRendering.h
--- a/Engine/Source/Runtime/Engine/Public/ScreenRendering.h
+++ b/Engine/Source/Runtime/Engine/Public/ScreenRendering.h
@@ -80,7 +80,7 @@ public:
 		SetTextureParameter(RHICmdList, GetPixelShader(),InTexture,InTextureSampler,Texture);
 	}
 
-	void SetParameters(FRHICommandList& RHICmdList, FRHISamplerState* SamplerStateRHI, FTextureRHIParamRef TextureRHI)
+	void SetParameters(FRHICommandList& RHICmdList, FRHISamplerState* SamplerStateRHI, FRHITexture* TextureRHI)
 	{
 		SetTextureParameter(RHICmdList, GetPixelShader(),InTexture,InTextureSampler,SamplerStateRHI,TextureRHI);
 	}
@@ -121,7 +121,7 @@ public:
 		SetTextureParameter(RHICmdList, GetPixelShader(), InTexture, InTextureSampler, Texture);
 	}
 
-	void SetParameters(FRHICommandList& RHICmdList, FRHISamplerState* SamplerStateRHI, FTextureRHIParamRef TextureRHI)
+	void SetParameters(FRHICommandList& RHICmdList, FRHISamplerState* SamplerStateRHI, FRHITexture* TextureRHI)
 	{
 		SetTextureParameter(RHICmdList, GetPixelShader(), InTexture, InTextureSampler, SamplerStateRHI, TextureRHI);
 	}
@@ -160,7 +160,7 @@ public:
         SetTextureParameter(RHICmdList, GetPixelShader(),InTexture,InTextureSampler,Texture);
     }
 
-    void SetParameters(FRHICommandList& RHICmdList, FRHISamplerState* SamplerStateRHI, FTextureRHIParamRef TextureRHI)
+    void SetParameters(FRHICommandList& RHICmdList, FRHISamplerState* SamplerStateRHI, FRHITexture* TextureRHI)
     {
         SetTextureParameter(RHICmdList, GetPixelShader(),InTexture,InTextureSampler,SamplerStateRHI,TextureRHI);
     }
@@ -195,7 +195,7 @@ public:
 	FScreenVS() {}
 
 
-	void SetParameters(FRHICommandList& RHICmdList, const FUniformBufferRHIParamRef ViewUniformBuffer)
+	void SetParameters(FRHICommandList& RHICmdList, FRHIUniformBuffer* ViewUniformBuffer)
 	{
 		FGlobalShader::SetParameters<FViewUniformShaderParameters>(RHICmdList, GetVertexShader(), ViewUniformBuffer);
 	}
diff --git a/Engine/Source/Runtime/Engine/Public/StereoLayerRendering.h b/Engine/Source/Runtime/Engine/Public/StereoLayerRendering.h
--- a/Engine/Source/Runtime/Engine/Public/StereoLayerRendering.h
+++ b/Engine/Source/Runtime/Engine/Public/StereoLayerRendering.h
@@ -85,7 +85,7 @@ public:
 
 	static bool ShouldCompilePermutation(const FGlobalShaderPermutationParameters& Parameters) { return true; }
 
-	void SetParameters(FRHICommandList& RHICmdList, FRHISamplerState* SamplerStateRHI, FTextureRHIParamRef TextureRHI)
+	void SetParameters(FRHICommandList& RHICmdList, FRHISamplerState* SamplerStateRHI, FRHITexture* TextureRHI)
 	{
 		FRHIPixelShader* PS = GetPixelShader();
 
diff --git a/Engine/Source/Runtime/GameplayMediaEncoder/Private/Microsoft/Windows/UtilityShaders.h b/Engine/Source/Runtime/GameplayMediaEncoder/Private/Microsoft/Windows/UtilityShaders.h
--- a/Engine/Source/Runtime/GameplayMediaEncoder/Private/Microsoft/Windows/UtilityShaders.h
+++ b/Engine/Source/Runtime/GameplayMediaEncoder/Private/Microsoft/Windows/UtilityShaders.h
@@ -38,7 +38,7 @@ public:
 		SetTextureParameter(RHICmdList, GetPixelShader(), InTexture, InTextureSampler, Texture);
 	}
 
-	void SetParameters(FRHICommandList& RHICmdList, FRHISamplerState* SamplerStateRHI, FTextureRHIParamRef TextureRHI)
+	void SetParameters(FRHICommandList& RHICmdList, FRHISamplerState* SamplerStateRHI, FRHITexture* TextureRHI)
 	{
 		SetTextureParameter(RHICmdList, GetPixelShader(), InTexture, InTextureSampler, SamplerStateRHI, TextureRHI);
 	}
diff --git a/Engine/Source/Runtime/HeadMountedDisplay/Private/XRSwapChain.cpp b/Engine/Source/Runtime/HeadMountedDisplay/Private/XRSwapChain.cpp
--- a/Engine/Source/Runtime/HeadMountedDisplay/Private/XRSwapChain.cpp
+++ b/Engine/Source/Runtime/HeadMountedDisplay/Private/XRSwapChain.cpp
@@ -60,7 +60,7 @@ FORCEINLINE void CheckInRHIThread()
 #endif
 }
 
-FXRSwapChain::FXRSwapChain(FTextureRHIParamRef InRHITexture, TArray<FTextureRHIRef>&& InRHITextureSwapChain)
+FXRSwapChain::FXRSwapChain(FRHITexture* InRHITexture, TArray<FTextureRHIRef>&& InRHITextureSwapChain)
 	: RHITexture(InRHITexture)
 	, RHITextureSwapChain(InRHITextureSwapChain)
 	, SwapChainIndex_RHIThread(0)
diff --git a/Engine/Source/Runtime/HeadMountedDisplay/Public/HeadMountedDisplayBase.h b/Engine/Source/Runtime/HeadMountedDisplay/Public/HeadMountedDisplayBase.h
--- a/Engine/Source/Runtime/HeadMountedDisplay/Public/HeadMountedDisplayBase.h
+++ b/Engine/Source/Runtime/HeadMountedDisplay/Public/HeadMountedDisplayBase.h
@@ -55,7 +55,7 @@ public:
 	// Get the rectangle of the HMD rendertarget for the left eye which seems undistorted enough to be cropped and displayed on the spectator screen.
 	virtual FIntRect GetFullFlatEyeRect_RenderThread(FTexture2DRHIRef EyeTexture) const { return FIntRect(0, 0, 1, 1); }
 	// Helper to copy one render target into another for spectator screen display
-	virtual void CopyTexture_RenderThread(FRHICommandListImmediate& RHICmdList, FTexture2DRHIParamRef SrcTexture, FIntRect SrcRect, FTexture2DRHIParamRef DstTexture, FIntRect DstRect, bool bClearBlack, bool bNoAlpha) const {}
+	virtual void CopyTexture_RenderThread(FRHICommandListImmediate& RHICmdList, FRHITexture2D* SrcTexture, FIntRect SrcRect, FRHITexture2D* DstTexture, FIntRect DstRect, bool bClearBlack, bool bNoAlpha) const {}
 
 protected:
 	/**
diff --git a/Engine/Source/Runtime/HeadMountedDisplay/Public/XRSwapChain.h b/Engine/Source/Runtime/HeadMountedDisplay/Public/XRSwapChain.h
--- a/Engine/Source/Runtime/HeadMountedDisplay/Public/XRSwapChain.h
+++ b/Engine/Source/Runtime/HeadMountedDisplay/Public/XRSwapChain.h
@@ -11,7 +11,7 @@
 class HEADMOUNTEDDISPLAY_API FXRSwapChain : public TSharedFromThis<FXRSwapChain, ESPMode::ThreadSafe>
 {
 public:
-	FXRSwapChain(FTextureRHIParamRef InRHITexture, TArray<FTextureRHIRef>&& InRHITextureSwapChain);
+	FXRSwapChain(FRHITexture* InRHITexture, TArray<FTextureRHIRef>&& InRHITextureSwapChain);
 	virtual ~FXRSwapChain();
 
 	FRHITexture* GetTexture() const { return RHITexture.GetReference(); }
diff --git a/Engine/Source/Runtime/NullDrv/Public/NullRHI.h b/Engine/Source/Runtime/NullDrv/Public/NullRHI.h
--- a/Engine/Source/Runtime/NullDrv/Public/NullRHI.h
+++ b/Engine/Source/Runtime/NullDrv/Public/NullRHI.h
@@ -93,7 +93,7 @@ public:
 
 	}
 
-	virtual void RHIDispatchIndirectComputeShader(FVertexBufferRHIParamRef ArgumentBuffer, uint32 ArgumentOffset) final override
+	virtual void RHIDispatchIndirectComputeShader(FRHIVertexBuffer* ArgumentBuffer, uint32 ArgumentOffset) final override
 	{
 
 	}
@@ -118,7 +118,7 @@ public:
 		return new FRHIUniformBuffer(Layout); 
 	}
 
-	virtual void RHIUpdateUniformBuffer(FUniformBufferRHIParamRef UniformBufferRHI, const void* Contents) final override
+	virtual void RHIUpdateUniformBuffer(FRHIUniformBuffer* UniformBufferRHI, const void* Contents) final override
 	{
 
 	}
@@ -132,16 +132,16 @@ public:
 		return new FRHIIndexBuffer(Stride,Size,InUsage); 
 	}
 
-	virtual void* RHILockIndexBuffer(FIndexBufferRHIParamRef IndexBuffer, uint32 Offset, uint32 Size, EResourceLockMode LockMode) final override
+	virtual void* RHILockIndexBuffer(FRHIIndexBuffer* IndexBuffer, uint32 Offset, uint32 Size, EResourceLockMode LockMode) final override
 	{ 
 		return GetStaticBuffer(); 
 	}
-	virtual void RHIUnlockIndexBuffer(FIndexBufferRHIParamRef IndexBuffer) final override
+	virtual void RHIUnlockIndexBuffer(FRHIIndexBuffer* IndexBuffer) final override
 	{
 
 	}
 
-	virtual void RHITransferIndexBufferUnderlyingResource(FIndexBufferRHIParamRef DestIndexBuffer, FIndexBufferRHIParamRef SrcIndexBuffer) final override
+	virtual void RHITransferIndexBufferUnderlyingResource(FRHIIndexBuffer* DestIndexBuffer, FRHIIndexBuffer* SrcIndexBuffer) final override
 	{
 
 	}
@@ -155,22 +155,22 @@ public:
 		return new FRHIVertexBuffer(Size,InUsage); 
 	}
 
-	virtual void* RHILockVertexBuffer(FVertexBufferRHIParamRef VertexBuffer, uint32 Offset, uint32 SizeRHI, EResourceLockMode LockMode) final override
+	virtual void* RHILockVertexBuffer(FRHIVertexBuffer* VertexBuffer, uint32 Offset, uint32 SizeRHI, EResourceLockMode LockMode) final override
 	{ 
 		return GetStaticBuffer(); 
 	}
-	virtual void RHIUnlockVertexBuffer(FVertexBufferRHIParamRef VertexBuffer) final override
+	virtual void RHIUnlockVertexBuffer(FRHIVertexBuffer* VertexBuffer) final override
 	{
 
 	}
 
-	virtual void RHITransferVertexBufferUnderlyingResource(FVertexBufferRHIParamRef DestVertexBuffer, FVertexBufferRHIParamRef SrcVertexBuffer) final override
+	virtual void RHITransferVertexBufferUnderlyingResource(FRHIVertexBuffer* DestVertexBuffer, FRHIVertexBuffer* SrcVertexBuffer) final override
 	{
 
 	}
 
 
-	virtual void RHICopyVertexBuffer(FVertexBufferRHIParamRef SourceBuffer,FVertexBufferRHIParamRef DestBuffer) final override
+	virtual void RHICopyVertexBuffer(FRHIVertexBuffer* SourceBuffer, FRHIVertexBuffer* DestBuffer) final override
 	{
 
 	}
@@ -184,50 +184,50 @@ public:
 		return new FRHIStructuredBuffer(Stride,Size,InUsage); 
 	}
 
-	virtual void* RHILockStructuredBuffer(FStructuredBufferRHIParamRef StructuredBuffer, uint32 Offset, uint32 SizeRHI, EResourceLockMode LockMode) final override
+	virtual void* RHILockStructuredBuffer(FRHIStructuredBuffer* StructuredBuffer, uint32 Offset, uint32 SizeRHI, EResourceLockMode LockMode) final override
 	{ 
 		return GetStaticBuffer(); 
 	}
-	virtual void RHIUnlockStructuredBuffer(FStructuredBufferRHIParamRef StructuredBuffer) final override
+	virtual void RHIUnlockStructuredBuffer(FRHIStructuredBuffer* StructuredBuffer) final override
 	{
 
 	}
 
 
-	virtual FUnorderedAccessViewRHIRef RHICreateUnorderedAccessView(FStructuredBufferRHIParamRef StructuredBuffer, bool bUseUAVCounter, bool bAppendBuffer) final override
+	virtual FUnorderedAccessViewRHIRef RHICreateUnorderedAccessView(FRHIStructuredBuffer* StructuredBuffer, bool bUseUAVCounter, bool bAppendBuffer) final override
 	{ 
 		return new FRHIUnorderedAccessView(); 
 	}
 
 
-	virtual FUnorderedAccessViewRHIRef RHICreateUnorderedAccessView(FTextureRHIParamRef Texture, uint32 MipLevel) final override
+	virtual FUnorderedAccessViewRHIRef RHICreateUnorderedAccessView(FRHITexture* Texture, uint32 MipLevel) final override
 	{ 
 		return new FRHIUnorderedAccessView(); 
 	}
 
 
-	virtual FUnorderedAccessViewRHIRef RHICreateUnorderedAccessView(FVertexBufferRHIParamRef VertexBuffer, uint8 Format) final override
+	virtual FUnorderedAccessViewRHIRef RHICreateUnorderedAccessView(FRHIVertexBuffer* VertexBuffer, uint8 Format) final override
 	{ 
 		return new FRHIUnorderedAccessView(); 
 	}
 
-	virtual FUnorderedAccessViewRHIRef RHICreateUnorderedAccessView(FIndexBufferRHIParamRef IndexBuffer, uint8 Format) final override
+	virtual FUnorderedAccessViewRHIRef RHICreateUnorderedAccessView(FRHIIndexBuffer* IndexBuffer, uint8 Format) final override
 	{
 		return new FRHIUnorderedAccessView();
 	}
 
-	virtual FShaderResourceViewRHIRef RHICreateShaderResourceView(FStructuredBufferRHIParamRef StructuredBuffer) final override
+	virtual FShaderResourceViewRHIRef RHICreateShaderResourceView(FRHIStructuredBuffer* StructuredBuffer) final override
 	{ 
 		return new FRHIShaderResourceView(); 
 	}
 
 
-	virtual FShaderResourceViewRHIRef RHICreateShaderResourceView(FVertexBufferRHIParamRef VertexBuffer, uint32 Stride, uint8 Format) final override
+	virtual FShaderResourceViewRHIRef RHICreateShaderResourceView(FRHIVertexBuffer* VertexBuffer, uint32 Stride, uint8 Format) final override
 	{ 
 		return new FRHIShaderResourceView(); 
 	}
 
-	virtual FShaderResourceViewRHIRef RHICreateShaderResourceView(FIndexBufferRHIParamRef Buffer) final override
+	virtual FShaderResourceViewRHIRef RHICreateShaderResourceView(FRHIIndexBuffer* Buffer) final override
 	{ 
 		return new FRHIShaderResourceView(); 
 	}
@@ -271,7 +271,7 @@ public:
 		return new FRHITextureReferenceNullImpl(); 
 	}
 
-	virtual void RHIUpdateTextureReference(FRHITextureReference* TextureRef, FTextureRHIParamRef NewTexture) final override
+	virtual void RHIUpdateTextureReference(FRHITextureReference* TextureRef, FRHITexture* NewTexture) final override
 	{ 
 		if(TextureRef) 
 		{ 
@@ -290,7 +290,7 @@ public:
 		return FTexture2DRHIRef(); 
 	}
 
-	virtual void RHICopySharedMips(FTexture2DRHIParamRef DestTexture2D, FTexture2DRHIParamRef SrcTexture2D) final override
+	virtual void RHICopySharedMips(FRHITexture2D* DestTexture2D, FRHITexture2D* SrcTexture2D) final override
 	{
 	}
 	virtual FTexture2DArrayRHIRef RHICreateTexture2DArray(uint32 SizeX, uint32 SizeY, uint32 SizeZ, uint8 Format, uint32 NumMips, uint32 NumSamples, uint32 Flags, FRHIResourceCreateInfo& CreateInfo) final override
@@ -302,58 +302,58 @@ public:
 	{ 
 		return new FRHITexture3D(SizeX, SizeY, SizeZ, NumMips, (EPixelFormat)Format, Flags, CreateInfo.ClearValueBinding);
 	}
-	virtual void RHIGetResourceInfo(FTextureRHIParamRef Ref, FRHIResourceInfo& OutInfo) final override
+	virtual void RHIGetResourceInfo(FRHITexture* Ref, FRHIResourceInfo& OutInfo) final override
 	{
 	}
 
-	virtual FShaderResourceViewRHIRef RHICreateShaderResourceView(FTextureRHIParamRef Texture, const FRHITextureSRVCreateInfo& CreateInfo) final override
+	virtual FShaderResourceViewRHIRef RHICreateShaderResourceView(FRHITexture* Texture, const FRHITextureSRVCreateInfo& CreateInfo) final override
 	{ 
 		return new FRHIShaderResourceView(); 
 	}
 
-	virtual void RHIGenerateMips(FTextureRHIParamRef Texture) final override
+	virtual void RHIGenerateMips(FRHITexture* Texture) final override
 	{
 
 	}
-	virtual uint32 RHIComputeMemorySize(FTextureRHIParamRef TextureRHI) final override
+	virtual uint32 RHIComputeMemorySize(FRHITexture* TextureRHI) final override
 	{ 
 		return 0; 
 	}
-	virtual FTexture2DRHIRef RHIAsyncReallocateTexture2D(FTexture2DRHIParamRef Texture2D, int32 NewMipCount, int32 NewSizeX, int32 NewSizeY, FThreadSafeCounter* RequestStatus) final override
+	virtual FTexture2DRHIRef RHIAsyncReallocateTexture2D(FRHITexture2D* Texture2D, int32 NewMipCount, int32 NewSizeX, int32 NewSizeY, FThreadSafeCounter* RequestStatus) final override
 	{ 
 		return new FRHITexture2D(NewSizeX,NewSizeY,NewMipCount,1,Texture2D->GetFormat(),Texture2D->GetFlags(), Texture2D->GetClearBinding());
 	}
-	virtual ETextureReallocationStatus RHIFinalizeAsyncReallocateTexture2D(FTexture2DRHIParamRef Texture2D, bool bBlockUntilCompleted) final override
+	virtual ETextureReallocationStatus RHIFinalizeAsyncReallocateTexture2D(FRHITexture2D* Texture2D, bool bBlockUntilCompleted) final override
 	{ 
 		return TexRealloc_Succeeded; 
 	}
-	virtual ETextureReallocationStatus RHICancelAsyncReallocateTexture2D(FTexture2DRHIParamRef Texture2D, bool bBlockUntilCompleted) final override
+	virtual ETextureReallocationStatus RHICancelAsyncReallocateTexture2D(FRHITexture2D* Texture2D, bool bBlockUntilCompleted) final override
 	{ 
 		return TexRealloc_Succeeded; 
 	}
-	virtual void* RHILockTexture2D(FTexture2DRHIParamRef Texture, uint32 MipIndex, EResourceLockMode LockMode, uint32& DestStride, bool bLockWithinMiptail) final override
+	virtual void* RHILockTexture2D(FRHITexture2D* Texture, uint32 MipIndex, EResourceLockMode LockMode, uint32& DestStride, bool bLockWithinMiptail) final override
 	{ 
 		DestStride = 0; 
 		return GetStaticBuffer(); 
 	}
-	virtual void RHIUnlockTexture2D(FTexture2DRHIParamRef Texture, uint32 MipIndex, bool bLockWithinMiptail) final override
+	virtual void RHIUnlockTexture2D(FRHITexture2D* Texture, uint32 MipIndex, bool bLockWithinMiptail) final override
 	{
 
 	}
-	virtual void* RHILockTexture2DArray(FTexture2DArrayRHIParamRef Texture, uint32 TextureIndex, uint32 MipIndex, EResourceLockMode LockMode, uint32& DestStride, bool bLockWithinMiptail) final override
+	virtual void* RHILockTexture2DArray(FRHITexture2DArray* Texture, uint32 TextureIndex, uint32 MipIndex, EResourceLockMode LockMode, uint32& DestStride, bool bLockWithinMiptail) final override
 	{ 
 		DestStride = 0; 
 		return GetStaticBuffer(); 
 	}
-	virtual void RHIUnlockTexture2DArray(FTexture2DArrayRHIParamRef Texture, uint32 TextureIndex, uint32 MipIndex, bool bLockWithinMiptail) final override
+	virtual void RHIUnlockTexture2DArray(FRHITexture2DArray* Texture, uint32 TextureIndex, uint32 MipIndex, bool bLockWithinMiptail) final override
 	{
 
 	}
-	virtual void RHIUpdateTexture2D(FTexture2DRHIParamRef Texture, uint32 MipIndex, const struct FUpdateTextureRegion2D& UpdateRegion, uint32 SourcePitch, const uint8* SourceData) final override
+	virtual void RHIUpdateTexture2D(FRHITexture2D* Texture, uint32 MipIndex, const struct FUpdateTextureRegion2D& UpdateRegion, uint32 SourcePitch, const uint8* SourceData) final override
 	{
 
 	}
-	virtual void RHIUpdateTexture3D(FTexture3DRHIParamRef Texture, uint32 MipIndex, const struct FUpdateTextureRegion3D& UpdateRegion, uint32 SourceRowPitch, uint32 SourceDepthPitch, const uint8* SourceData) final override
+	virtual void RHIUpdateTexture3D(FRHITexture3D* Texture, uint32 MipIndex, const struct FUpdateTextureRegion3D& UpdateRegion, uint32 SourceRowPitch, uint32 SourceDepthPitch, const uint8* SourceData) final override
 	{
 
 	}
@@ -365,52 +365,52 @@ public:
 	{ 
 		return new FRHITextureCube(Size, NumMips, (EPixelFormat)Format, Flags, CreateInfo.ClearValueBinding);
 	}
-	virtual void* RHILockTextureCubeFace(FTextureCubeRHIParamRef Texture, uint32 FaceIndex, uint32 ArrayIndex, uint32 MipIndex, EResourceLockMode LockMode, uint32& DestStride, bool bLockWithinMiptail) final override
+	virtual void* RHILockTextureCubeFace(FRHITextureCube* Texture, uint32 FaceIndex, uint32 ArrayIndex, uint32 MipIndex, EResourceLockMode LockMode, uint32& DestStride, bool bLockWithinMiptail) final override
 	{ 
 		DestStride = 0; 
 		return GetStaticBuffer(); 
 	}
-	virtual void RHIUnlockTextureCubeFace(FTextureCubeRHIParamRef Texture, uint32 FaceIndex, uint32 ArrayIndex, uint32 MipIndex, bool bLockWithinMiptail) final override
+	virtual void RHIUnlockTextureCubeFace(FRHITextureCube* Texture, uint32 FaceIndex, uint32 ArrayIndex, uint32 MipIndex, bool bLockWithinMiptail) final override
 	{
 	}
-	virtual void RHICopyToResolveTarget(FTextureRHIParamRef SourceTexture, FTextureRHIParamRef DestTexture, const FResolveParams& ResolveParams) final override
+	virtual void RHICopyToResolveTarget(FRHITexture* SourceTexture, FRHITexture* DestTexture, const FResolveParams& ResolveParams) final override
 	{
 
 	}
 
-	virtual void RHICopyTexture(FTextureRHIParamRef SourceTexture, FTextureRHIParamRef DestTexture, const FRHICopyTextureInfo& CopyInfo) final override
+	virtual void RHICopyTexture(FRHITexture* SourceTexture, FRHITexture* DestTexture, const FRHICopyTextureInfo& CopyInfo) final override
 	{
 
 	}
 
-	virtual void RHIBindDebugLabelName(FTextureRHIParamRef Texture, const TCHAR* Name) final override
+	virtual void RHIBindDebugLabelName(FRHITexture* Texture, const TCHAR* Name) final override
 	{
 
 	}
 
-	virtual void RHIReadSurfaceData(FTextureRHIParamRef Texture,FIntRect Rect,TArray<FColor>& OutData,FReadSurfaceDataFlags InFlags) final override
+	virtual void RHIReadSurfaceData(FRHITexture* Texture,FIntRect Rect,TArray<FColor>& OutData,FReadSurfaceDataFlags InFlags) final override
 	{ 
 		OutData.AddZeroed(Rect.Width() * Rect.Height()); 
 	}
 
 
-	virtual void RHIMapStagingSurface(FTextureRHIParamRef Texture,void*& OutData,int32& OutWidth,int32& OutHeight) final override
+	virtual void RHIMapStagingSurface(FRHITexture* Texture,void*& OutData,int32& OutWidth,int32& OutHeight) final override
 	{
 
 	}
 
 
-	virtual void RHIUnmapStagingSurface(FTextureRHIParamRef Texture) final override
+	virtual void RHIUnmapStagingSurface(FRHITexture* Texture) final override
 	{
 
 	}
 
-	virtual void RHIReadSurfaceFloatData(FTextureRHIParamRef Texture,FIntRect Rect,TArray<FFloat16Color>& OutData,ECubeFace CubeFace,int32 ArrayIndex,int32 MipIndex) final override
+	virtual void RHIReadSurfaceFloatData(FRHITexture* Texture,FIntRect Rect,TArray<FFloat16Color>& OutData,ECubeFace CubeFace,int32 ArrayIndex,int32 MipIndex) final override
 	{
 
 	}
 
-	virtual void RHIRead3DSurfaceFloatData(FTextureRHIParamRef Texture,FIntRect Rect,FIntPoint ZMinMax,TArray<FFloat16Color>& OutData) final override
+	virtual void RHIRead3DSurfaceFloatData(FRHITexture* Texture,FIntRect Rect,FIntPoint ZMinMax,TArray<FFloat16Color>& OutData) final override
 	{
 
 	}
@@ -442,7 +442,7 @@ public:
 	}
 
 
-	virtual void RHIBeginDrawingViewport(FRHIViewport* Viewport, FTextureRHIParamRef RenderTargetRHI) final override
+	virtual void RHIBeginDrawingViewport(FRHIViewport* Viewport, FRHITexture* RenderTargetRHI) final override
 	{
 	}
 
@@ -473,7 +473,7 @@ public:
 	{
 
 	}
-	virtual void RHIAliasTextureResources(FTextureRHIParamRef DestTexture, FTextureRHIParamRef SrcTexture) final override
+	virtual void RHIAliasTextureResources(FRHITexture* DestTexture, FRHITexture* SrcTexture) final override
 	{
 
 	}
@@ -514,11 +514,11 @@ public:
 
 	}
 
-	virtual void RHISetStreamSource(uint32 StreamIndex, FVertexBufferRHIParamRef VertexBuffer, uint32 Offset) final override
+	virtual void RHISetStreamSource(uint32 StreamIndex, FRHIVertexBuffer* VertexBuffer, uint32 Offset) final override
 	{
 	}
 
-	virtual void RHISetStreamOutTargets(uint32 NumTargets,const FVertexBufferRHIParamRef* VertexBuffers,const uint32* Offsets) final override
+	virtual void RHISetStreamOutTargets(uint32 NumTargets, FRHIVertexBuffer* const* VertexBuffers,const uint32* Offsets) final override
 	{
 
 	}
@@ -541,37 +541,37 @@ public:
 	}
 
 
-	virtual void RHISetShaderTexture(FRHIVertexShader* VertexShader, uint32 TextureIndex, FTextureRHIParamRef NewTexture) final override
+	virtual void RHISetShaderTexture(FRHIVertexShader* VertexShader, uint32 TextureIndex, FRHITexture* NewTexture) final override
 	{
 
 	}
 
 
-	virtual void RHISetShaderTexture(FRHIHullShader* HullShader, uint32 TextureIndex, FTextureRHIParamRef NewTexture) final override
+	virtual void RHISetShaderTexture(FRHIHullShader* HullShader, uint32 TextureIndex, FRHITexture* NewTexture) final override
 	{
 
 	}
 
 
-	virtual void RHISetShaderTexture(FRHIDomainShader* DomainShader, uint32 TextureIndex, FTextureRHIParamRef NewTexture) final override
+	virtual void RHISetShaderTexture(FRHIDomainShader* DomainShader, uint32 TextureIndex, FRHITexture* NewTexture) final override
 	{
 
 	}
 
 
-	virtual void RHISetShaderTexture(FRHIGeometryShader* GeometryShader, uint32 TextureIndex, FTextureRHIParamRef NewTexture) final override
+	virtual void RHISetShaderTexture(FRHIGeometryShader* GeometryShader, uint32 TextureIndex, FRHITexture* NewTexture) final override
 	{
 
 	}
 
 
-	virtual void RHISetShaderTexture(FRHIPixelShader* PixelShader, uint32 TextureIndex, FTextureRHIParamRef NewTexture) final override
+	virtual void RHISetShaderTexture(FRHIPixelShader* PixelShader, uint32 TextureIndex, FRHITexture* NewTexture) final override
 	{
 
 	}
 
 
-	virtual void RHISetShaderTexture(FRHIComputeShader* PixelShader, uint32 TextureIndex, FTextureRHIParamRef NewTexture) final override
+	virtual void RHISetShaderTexture(FRHIComputeShader* PixelShader, uint32 TextureIndex, FRHITexture* NewTexture) final override
 	{
 
 	}
@@ -650,32 +650,32 @@ public:
 
 	}
 
-	virtual void RHISetShaderUniformBuffer(FRHIVertexShader* VertexShader, uint32 BufferIndex, FUniformBufferRHIParamRef Buffer) final override
+	virtual void RHISetShaderUniformBuffer(FRHIVertexShader* VertexShader, uint32 BufferIndex, FRHIUniformBuffer* Buffer) final override
 	{
 
 	}
 
-	virtual void RHISetShaderUniformBuffer(FRHIHullShader* HullShader, uint32 BufferIndex, FUniformBufferRHIParamRef Buffer) final override
+	virtual void RHISetShaderUniformBuffer(FRHIHullShader* HullShader, uint32 BufferIndex, FRHIUniformBuffer* Buffer) final override
 	{
 
 	}
 
-	virtual void RHISetShaderUniformBuffer(FRHIDomainShader* DomainShader, uint32 BufferIndex, FUniformBufferRHIParamRef Buffer) final override
+	virtual void RHISetShaderUniformBuffer(FRHIDomainShader* DomainShader, uint32 BufferIndex, FRHIUniformBuffer* Buffer) final override
 	{
 
 	}
 
-	virtual void RHISetShaderUniformBuffer(FRHIGeometryShader* GeometryShader, uint32 BufferIndex, FUniformBufferRHIParamRef Buffer) final override
+	virtual void RHISetShaderUniformBuffer(FRHIGeometryShader* GeometryShader, uint32 BufferIndex, FRHIUniformBuffer* Buffer) final override
 	{
 
 	}
 
-	virtual void RHISetShaderUniformBuffer(FRHIPixelShader* PixelShader, uint32 BufferIndex, FUniformBufferRHIParamRef Buffer) final override
+	virtual void RHISetShaderUniformBuffer(FRHIPixelShader* PixelShader, uint32 BufferIndex, FRHIUniformBuffer* Buffer) final override
 	{
 
 	}
 
-	virtual void RHISetShaderUniformBuffer(FRHIComputeShader* ComputeShader, uint32 BufferIndex, FUniformBufferRHIParamRef Buffer) final override
+	virtual void RHISetShaderUniformBuffer(FRHIComputeShader* ComputeShader, uint32 BufferIndex, FRHIUniformBuffer* Buffer) final override
 	{
 
 	}
@@ -731,22 +731,22 @@ public:
 	{
 
 	}
-	virtual void RHIDrawPrimitiveIndirect(FVertexBufferRHIParamRef ArgumentBuffer, uint32 ArgumentOffset) final override
+	virtual void RHIDrawPrimitiveIndirect(FRHIVertexBuffer* ArgumentBuffer, uint32 ArgumentOffset) final override
 	{
 
 	}
 
-	virtual void RHIDrawIndexedIndirect(FIndexBufferRHIParamRef IndexBufferRHI, FStructuredBufferRHIParamRef ArgumentsBufferRHI, int32 DrawArgumentsIndex, uint32 NumInstances) final override
+	virtual void RHIDrawIndexedIndirect(FRHIIndexBuffer* IndexBufferRHI, FRHIStructuredBuffer* ArgumentsBufferRHI, int32 DrawArgumentsIndex, uint32 NumInstances) final override
 	{
 
 	}
 
 
-	virtual void RHIDrawIndexedPrimitive(FIndexBufferRHIParamRef IndexBuffer, int32 BaseVertexIndex, uint32 FirstInstance, uint32 NumVertices, uint32 StartIndex, uint32 NumPrimitives, uint32 NumInstances) final override
+	virtual void RHIDrawIndexedPrimitive(FRHIIndexBuffer* IndexBuffer, int32 BaseVertexIndex, uint32 FirstInstance, uint32 NumVertices, uint32 StartIndex, uint32 NumPrimitives, uint32 NumInstances) final override
 	{
 
 	}
-	virtual void RHIDrawIndexedPrimitiveIndirect(FIndexBufferRHIParamRef IndexBuffer, FVertexBufferRHIParamRef ArgumentBuffer, uint32 ArgumentOffset) final override
+	virtual void RHIDrawIndexedPrimitiveIndirect(FRHIIndexBuffer* IndexBuffer, FRHIVertexBuffer* ArgumentBuffer, uint32 ArgumentOffset) final override
 	{
 
 	}
@@ -770,11 +770,11 @@ public:
 	{
 
 	}
-	virtual void RHIVirtualTextureSetFirstMipInMemory(FTexture2DRHIParamRef Texture, uint32 FirstMip) final override
+	virtual void RHIVirtualTextureSetFirstMipInMemory(FRHITexture2D* Texture, uint32 FirstMip) final override
 	{
 
 	}
-	virtual void RHIVirtualTextureSetFirstMipVisible(FTexture2DRHIParamRef Texture, uint32 FirstMip) final override
+	virtual void RHIVirtualTextureSetFirstMipVisible(FRHITexture2D* Texture, uint32 FirstMip) final override
 	{
 
 	}
diff --git a/Engine/Source/Runtime/OpenGLDrv/Private/OpenGLCommands.cpp b/Engine/Source/Runtime/OpenGLDrv/Private/OpenGLCommands.cpp
--- a/Engine/Source/Runtime/OpenGLDrv/Private/OpenGLCommands.cpp
+++ b/Engine/Source/Runtime/OpenGLDrv/Private/OpenGLCommands.cpp
@@ -273,7 +273,7 @@ static FORCEINLINE GLint ModifyFilterByMips(GLint Filter, bool bHasMips)
 }
 
 // Vertex state.
-void FOpenGLDynamicRHI::RHISetStreamSource(uint32 StreamIndex, FVertexBufferRHIParamRef VertexBufferRHI, uint32 Offset)
+void FOpenGLDynamicRHI::RHISetStreamSource(uint32 StreamIndex, FRHIVertexBuffer* VertexBufferRHI, uint32 Offset)
 {
 	VERIFY_GL_SCOPE();
 	FOpenGLVertexBuffer* VertexBuffer = ResourceCast(VertexBufferRHI);
@@ -282,7 +282,7 @@ void FOpenGLDynamicRHI::RHISetStreamSource(uint32 StreamIndex, FVertexBufferRHIP
 	PendingState.Streams[StreamIndex].Offset = Offset;
 }
 
-void FOpenGLDynamicRHI::RHISetStreamOutTargets(uint32 NumTargets, const FVertexBufferRHIParamRef* VertexBuffers, const uint32* Offsets)
+void FOpenGLDynamicRHI::RHISetStreamOutTargets(uint32 NumTargets, FRHIVertexBuffer* const* VertexBuffers, const uint32* Offsets)
 {
 	check(0);
 }
@@ -978,7 +978,7 @@ void FOpenGLDynamicRHI::RHISetShaderResourceViewParameter(FRHIGeometryShader* Ge
 	RHISetShaderSampler(GeometryShaderRHI,TextureIndex,PointSamplerState);
 }
 
-void FOpenGLDynamicRHI::RHISetShaderTexture(FRHIVertexShader* VertexShaderRHI,uint32 TextureIndex,FTextureRHIParamRef NewTextureRHI)
+void FOpenGLDynamicRHI::RHISetShaderTexture(FRHIVertexShader* VertexShaderRHI,uint32 TextureIndex, FRHITexture* NewTextureRHI)
 {
 	VALIDATE_BOUND_SHADER(VertexShaderRHI);
 
@@ -995,7 +995,7 @@ void FOpenGLDynamicRHI::RHISetShaderTexture(FRHIVertexShader* VertexShaderRHI,ui
 	}
 }
 
-void FOpenGLDynamicRHI::RHISetShaderTexture(FRHIHullShader* HullShaderRHI,uint32 TextureIndex,FTextureRHIParamRef NewTextureRHI)
+void FOpenGLDynamicRHI::RHISetShaderTexture(FRHIHullShader* HullShaderRHI,uint32 TextureIndex, FRHITexture* NewTextureRHI)
 {
 	VALIDATE_BOUND_SHADER(HullShaderRHI);
 	
@@ -1013,7 +1013,7 @@ void FOpenGLDynamicRHI::RHISetShaderTexture(FRHIHullShader* HullShaderRHI,uint32
 	}
 }
 
-void FOpenGLDynamicRHI::RHISetShaderTexture(FRHIDomainShader* DomainShaderRHI,uint32 TextureIndex,FTextureRHIParamRef NewTextureRHI)
+void FOpenGLDynamicRHI::RHISetShaderTexture(FRHIDomainShader* DomainShaderRHI,uint32 TextureIndex, FRHITexture* NewTextureRHI)
 {
 	VALIDATE_BOUND_SHADER(DomainShaderRHI);
 
@@ -1031,7 +1031,7 @@ void FOpenGLDynamicRHI::RHISetShaderTexture(FRHIDomainShader* DomainShaderRHI,ui
 	}
 }
 
-void FOpenGLDynamicRHI::RHISetShaderTexture(FRHIGeometryShader* GeometryShaderRHI,uint32 TextureIndex,FTextureRHIParamRef NewTextureRHI)
+void FOpenGLDynamicRHI::RHISetShaderTexture(FRHIGeometryShader* GeometryShaderRHI,uint32 TextureIndex, FRHITexture* NewTextureRHI)
 {
 	VALIDATE_BOUND_SHADER(GeometryShaderRHI);
 
@@ -1048,7 +1048,7 @@ void FOpenGLDynamicRHI::RHISetShaderTexture(FRHIGeometryShader* GeometryShaderRH
 	}
 }
 
-void FOpenGLDynamicRHI::RHISetShaderTexture(FRHIPixelShader* PixelShaderRHI,uint32 TextureIndex,FTextureRHIParamRef NewTextureRHI)
+void FOpenGLDynamicRHI::RHISetShaderTexture(FRHIPixelShader* PixelShaderRHI,uint32 TextureIndex, FRHITexture* NewTextureRHI)
 {
 	VALIDATE_BOUND_SHADER(PixelShaderRHI);
 
@@ -1102,7 +1102,7 @@ void FOpenGLDynamicRHI::RHISetShaderSampler(FRHIGeometryShader* GeometryShaderRH
 	InternalSetSamplerStates(FOpenGL::GetFirstGeometryTextureUnit() + SamplerIndex, NewState);
 }
 
-void FOpenGLDynamicRHI::RHISetShaderTexture(FRHIComputeShader* ComputeShaderRHI,uint32 TextureIndex,FTextureRHIParamRef NewTextureRHI)
+void FOpenGLDynamicRHI::RHISetShaderTexture(FRHIComputeShader* ComputeShaderRHI,uint32 TextureIndex, FRHITexture* NewTextureRHI)
 {
 	check(GMaxRHIFeatureLevel >= ERHIFeatureLevel::SM5);
 	VERIFY_GL_SCOPE();
@@ -1127,7 +1127,7 @@ void FOpenGLDynamicRHI::RHISetShaderSampler(FRHIPixelShader* PixelShaderRHI,uint
 	InternalSetSamplerStates(FOpenGL::GetFirstPixelTextureUnit() + SamplerIndex, NewState);
 }
 
-void FOpenGLDynamicRHI::RHISetShaderUniformBuffer(FRHIVertexShader* VertexShaderRHI,uint32 BufferIndex,FUniformBufferRHIParamRef BufferRHI)
+void FOpenGLDynamicRHI::RHISetShaderUniformBuffer(FRHIVertexShader* VertexShaderRHI,uint32 BufferIndex, FRHIUniformBuffer* BufferRHI)
 {
 	VALIDATE_BOUND_SHADER(VertexShaderRHI);
 	VERIFY_GL_SCOPE();
@@ -1136,7 +1136,7 @@ void FOpenGLDynamicRHI::RHISetShaderUniformBuffer(FRHIVertexShader* VertexShader
 	PendingState.bAnyDirtyGraphicsUniformBuffers = true;
 }
 
-void FOpenGLDynamicRHI::RHISetShaderUniformBuffer(FRHIHullShader* HullShaderRHI,uint32 BufferIndex,FUniformBufferRHIParamRef BufferRHI)
+void FOpenGLDynamicRHI::RHISetShaderUniformBuffer(FRHIHullShader* HullShaderRHI,uint32 BufferIndex, FRHIUniformBuffer* BufferRHI)
 {
 	VALIDATE_BOUND_SHADER(HullShaderRHI);
 	VERIFY_GL_SCOPE();
@@ -1146,7 +1146,7 @@ void FOpenGLDynamicRHI::RHISetShaderUniformBuffer(FRHIHullShader* HullShaderRHI,
 	PendingState.bAnyDirtyGraphicsUniformBuffers = true;
 }
 
-void FOpenGLDynamicRHI::RHISetShaderUniformBuffer(FRHIDomainShader* DomainShaderRHI,uint32 BufferIndex,FUniformBufferRHIParamRef BufferRHI)
+void FOpenGLDynamicRHI::RHISetShaderUniformBuffer(FRHIDomainShader* DomainShaderRHI,uint32 BufferIndex, FRHIUniformBuffer* BufferRHI)
 {
 	VALIDATE_BOUND_SHADER(DomainShaderRHI);
 	VERIFY_GL_SCOPE();
@@ -1156,7 +1156,7 @@ void FOpenGLDynamicRHI::RHISetShaderUniformBuffer(FRHIDomainShader* DomainShader
 	PendingState.bAnyDirtyGraphicsUniformBuffers = true;
 }
 
-void FOpenGLDynamicRHI::RHISetShaderUniformBuffer(FRHIGeometryShader* GeometryShaderRHI,uint32 BufferIndex,FUniformBufferRHIParamRef BufferRHI)
+void FOpenGLDynamicRHI::RHISetShaderUniformBuffer(FRHIGeometryShader* GeometryShaderRHI,uint32 BufferIndex, FRHIUniformBuffer* BufferRHI)
 {
 	VALIDATE_BOUND_SHADER(GeometryShaderRHI);
 	VERIFY_GL_SCOPE();
@@ -1180,7 +1180,7 @@ void FOpenGLDynamicRHI::RHISetShaderSampler(FRHIComputeShader* ComputeShaderRHI,
 	}
 }
 
-void FOpenGLDynamicRHI::RHISetShaderUniformBuffer(FRHIPixelShader* PixelShaderRHI,uint32 BufferIndex,FUniformBufferRHIParamRef BufferRHI)
+void FOpenGLDynamicRHI::RHISetShaderUniformBuffer(FRHIPixelShader* PixelShaderRHI,uint32 BufferIndex, FRHIUniformBuffer* BufferRHI)
 {
 	VERIFY_GL_SCOPE();
 	VALIDATE_BOUND_SHADER(PixelShaderRHI);
@@ -1189,7 +1189,7 @@ void FOpenGLDynamicRHI::RHISetShaderUniformBuffer(FRHIPixelShader* PixelShaderRH
 	PendingState.bAnyDirtyGraphicsUniformBuffers = true;
 }
 
-void FOpenGLDynamicRHI::RHISetShaderUniformBuffer(FRHIComputeShader* ComputeShaderRHI,uint32 BufferIndex,FUniformBufferRHIParamRef BufferRHI)
+void FOpenGLDynamicRHI::RHISetShaderUniformBuffer(FRHIComputeShader* ComputeShaderRHI,uint32 BufferIndex, FRHIUniformBuffer* BufferRHI)
 {
 	VERIFY_GL_SCOPE();
 	check(GMaxRHIFeatureLevel >= ERHIFeatureLevel::SM5);
@@ -2680,7 +2680,7 @@ void FOpenGLDynamicRHI::RHIDrawPrimitive(uint32 BaseVertexIndex,uint32 NumPrimit
 	}
 }
 
-void FOpenGLDynamicRHI::RHIDrawPrimitiveIndirect(FVertexBufferRHIParamRef ArgumentBufferRHI,uint32 ArgumentOffset)
+void FOpenGLDynamicRHI::RHIDrawPrimitiveIndirect(FRHIVertexBuffer* ArgumentBufferRHI,uint32 ArgumentOffset)
 {
 	if (FOpenGL::SupportsDrawIndirect())
 	{
@@ -2731,7 +2731,7 @@ void FOpenGLDynamicRHI::RHIDrawPrimitiveIndirect(FVertexBufferRHIParamRef Argume
 
 }
 
-void FOpenGLDynamicRHI::RHIDrawIndexedIndirect(FIndexBufferRHIParamRef IndexBufferRHI,  FStructuredBufferRHIParamRef ArgumentsBufferRHI, int32 DrawArgumentsIndex, uint32 NumInstances)
+void FOpenGLDynamicRHI::RHIDrawIndexedIndirect(FRHIIndexBuffer* IndexBufferRHI, FRHIStructuredBuffer* ArgumentsBufferRHI, int32 DrawArgumentsIndex, uint32 NumInstances)
 {
 	if (FOpenGL::SupportsDrawIndirect())
 	{
@@ -2790,7 +2790,7 @@ void FOpenGLDynamicRHI::RHIDrawIndexedIndirect(FIndexBufferRHIParamRef IndexBuff
 	}
 }
 
-void FOpenGLDynamicRHI::RHIDrawIndexedPrimitive(FIndexBufferRHIParamRef IndexBufferRHI,int32 BaseVertexIndex,uint32 FirstInstance,uint32 NumVertices,uint32 StartIndex,uint32 NumPrimitives,uint32 NumInstances)
+void FOpenGLDynamicRHI::RHIDrawIndexedPrimitive(FRHIIndexBuffer* IndexBufferRHI,int32 BaseVertexIndex,uint32 FirstInstance,uint32 NumVertices,uint32 StartIndex,uint32 NumPrimitives,uint32 NumInstances)
 {
 	SCOPE_CYCLE_COUNTER_DETAILED(STAT_OpenGLDrawPrimitiveTime);
 	VERIFY_GL_SCOPE();
@@ -2892,7 +2892,7 @@ void FOpenGLDynamicRHI::RHIDrawIndexedPrimitive(FIndexBufferRHIParamRef IndexBuf
 	}
 }
 
-void FOpenGLDynamicRHI::RHIDrawIndexedPrimitiveIndirect(FIndexBufferRHIParamRef IndexBufferRHI,FVertexBufferRHIParamRef ArgumentBufferRHI,uint32 ArgumentOffset)
+void FOpenGLDynamicRHI::RHIDrawIndexedPrimitiveIndirect(FRHIIndexBuffer* IndexBufferRHI, FRHIVertexBuffer* ArgumentBufferRHI,uint32 ArgumentOffset)
 {
 	if (FOpenGL::SupportsDrawIndirect())
 	{
@@ -3227,7 +3227,7 @@ void FOpenGLDynamicRHI::RHIDispatchComputeShader(uint32 ThreadGroupCountX, uint3
 	}
 }
 
-void FOpenGLDynamicRHI::RHIDispatchIndirectComputeShader(FVertexBufferRHIParamRef ArgumentBufferRHI,uint32 ArgumentOffset)
+void FOpenGLDynamicRHI::RHIDispatchIndirectComputeShader(FRHIVertexBuffer* ArgumentBufferRHI,uint32 ArgumentOffset)
 {
 	if (FOpenGL::SupportsComputeShaders())
 	{
@@ -3325,7 +3325,7 @@ void FOpenGLDynamicRHI::RHIInvalidateCachedState()
 	SharedContextState.InitializeResources(FOpenGL::GetMaxCombinedTextureImageUnits(), OGL_MAX_COMPUTE_STAGE_UAV_UNITS);
 }
 
-void FOpenGLDynamicRHI::RHICopyToStagingBuffer(FVertexBufferRHIParamRef SourceBufferRHI, FRHIStagingBuffer* DestinationStagingBufferRHI, uint32 InOffset, uint32 InNumBytes)
+void FOpenGLDynamicRHI::RHICopyToStagingBuffer(FRHIVertexBuffer* SourceBufferRHI, FRHIStagingBuffer* DestinationStagingBufferRHI, uint32 InOffset, uint32 InNumBytes)
 {
 #if OPENGL_GL3
 	VERIFY_GL_SCOPE();
diff --git a/Engine/Source/Runtime/OpenGLDrv/Private/OpenGLIndexBuffer.cpp b/Engine/Source/Runtime/OpenGLDrv/Private/OpenGLIndexBuffer.cpp
--- a/Engine/Source/Runtime/OpenGLDrv/Private/OpenGLIndexBuffer.cpp
+++ b/Engine/Source/Runtime/OpenGLDrv/Private/OpenGLIndexBuffer.cpp
@@ -34,7 +34,7 @@ FIndexBufferRHIRef FOpenGLDynamicRHI::CreateIndexBuffer_RenderThread(class FRHIC
 	return this->RHICreateIndexBuffer(Stride, Size, InUsage, CreateInfo);
 }
 
-void* FOpenGLDynamicRHI::RHILockIndexBuffer(FIndexBufferRHIParamRef IndexBufferRHI,uint32 Offset,uint32 Size,EResourceLockMode LockMode)
+void* FOpenGLDynamicRHI::RHILockIndexBuffer(FRHIIndexBuffer* IndexBufferRHI,uint32 Offset,uint32 Size,EResourceLockMode LockMode)
 {
 	FRHICommandListImmediate& RHICmdList = FRHICommandListExecutor::GetImmediateCommandList();
 	RHITHREAD_GLCOMMAND_PROLOGUE();
@@ -44,7 +44,7 @@ void* FOpenGLDynamicRHI::RHILockIndexBuffer(FIndexBufferRHIParamRef IndexBufferR
 	RHITHREAD_GLCOMMAND_EPILOGUE_RETURN(void*);
 }
 
-void FOpenGLDynamicRHI::RHIUnlockIndexBuffer(FIndexBufferRHIParamRef IndexBufferRHI)
+void FOpenGLDynamicRHI::RHIUnlockIndexBuffer(FRHIIndexBuffer* IndexBufferRHI)
 {
 	FRHICommandListImmediate& RHICmdList = FRHICommandListExecutor::GetImmediateCommandList();
 	RHITHREAD_GLCOMMAND_PROLOGUE();
diff --git a/Engine/Source/Runtime/OpenGLDrv/Private/OpenGLRenderTarget.cpp b/Engine/Source/Runtime/OpenGLDrv/Private/OpenGLRenderTarget.cpp
--- a/Engine/Source/Runtime/OpenGLDrv/Private/OpenGLRenderTarget.cpp
+++ b/Engine/Source/Runtime/OpenGLDrv/Private/OpenGLRenderTarget.cpp
@@ -364,7 +364,7 @@ GLuint FOpenGLDynamicRHI::GetOpenGLFramebuffer(uint32 NumSimultaneousRenderTarge
 	return Framebuffer;
 }
 
-void ReleaseOpenGLFramebuffers(FOpenGLDynamicRHI* Device, FTextureRHIParamRef TextureRHI)
+void ReleaseOpenGLFramebuffers(FOpenGLDynamicRHI* Device, FRHITexture* TextureRHI)
 {
 	FOpenGLTextureBase* Texture = GetOpenGLTextureFromRHITexture(TextureRHI);
 
@@ -436,7 +436,7 @@ void FOpenGLDynamicRHI::PurgeFramebufferFromCaches( GLuint Framebuffer )
 	}
 }
 
-void FOpenGLDynamicRHI::RHICopyToResolveTarget(FTextureRHIParamRef SourceTextureRHI, FTextureRHIParamRef DestTextureRHI, const FResolveParams& ResolveParams)
+void FOpenGLDynamicRHI::RHICopyToResolveTarget(FRHITexture* SourceTextureRHI, FRHITexture* DestTextureRHI, const FResolveParams& ResolveParams)
 {
 	if (!SourceTextureRHI || !DestTextureRHI)
 	{
@@ -574,7 +574,7 @@ void FOpenGLDynamicRHI::RHICopyToResolveTarget(FTextureRHIParamRef SourceTexture
 
 
 
-void FOpenGLDynamicRHI::ReadSurfaceDataRaw(FOpenGLContextState& ContextState, FTextureRHIParamRef TextureRHI,FIntRect Rect,TArray<uint8>& OutData, FReadSurfaceDataFlags InFlags)
+void FOpenGLDynamicRHI::ReadSurfaceDataRaw(FOpenGLContextState& ContextState, FRHITexture* TextureRHI,FIntRect Rect,TArray<uint8>& OutData, FReadSurfaceDataFlags InFlags)
 {
 	VERIFY_GL_SCOPE();
 
@@ -865,7 +865,7 @@ void FOpenGLDynamicRHI::ReadSurfaceDataRaw(FOpenGLContextState& ContextState, FT
 	ContextState.Framebuffer = (GLuint)-1;
 }
 
-void FOpenGLDynamicRHI::RHIReadSurfaceData(FTextureRHIParamRef TextureRHI,FIntRect Rect,TArray<FColor>& OutData, FReadSurfaceDataFlags InFlags)
+void FOpenGLDynamicRHI::RHIReadSurfaceData(FRHITexture* TextureRHI,FIntRect Rect,TArray<FColor>& OutData, FReadSurfaceDataFlags InFlags)
 {
 	if (!ensure(TextureRHI))
 	{
@@ -894,7 +894,7 @@ void FOpenGLDynamicRHI::RHIReadSurfaceData(FTextureRHIParamRef TextureRHI,FIntRe
 	RHITHREAD_GLCOMMAND_EPILOGUE();
 }
 
-void FOpenGLDynamicRHI::RHIReadSurfaceData(FTextureRHIParamRef TextureRHI, FIntRect Rect, TArray<FLinearColor>& OutData, FReadSurfaceDataFlags InFlags)
+void FOpenGLDynamicRHI::RHIReadSurfaceData(FRHITexture* TextureRHI, FIntRect Rect, TArray<FLinearColor>& OutData, FReadSurfaceDataFlags InFlags)
 {
 	VERIFY_GL_SCOPE();
 
@@ -936,7 +936,7 @@ void FOpenGLDynamicRHI::RHIReadSurfaceData(FTextureRHIParamRef TextureRHI, FIntR
 	GetContextStateForCurrentContext().Framebuffer = (GLuint)-1;
 }
 
-void FOpenGLDynamicRHI::RHIMapStagingSurface(FTextureRHIParamRef TextureRHI,void*& OutData,int32& OutWidth,int32& OutHeight)
+void FOpenGLDynamicRHI::RHIMapStagingSurface(FRHITexture* TextureRHI,void*& OutData,int32& OutWidth,int32& OutHeight)
 {
 	FRHICommandListImmediate& RHICmdList = FRHICommandListExecutor::GetImmediateCommandList();
 
@@ -956,7 +956,7 @@ void FOpenGLDynamicRHI::RHIMapStagingSurface(FTextureRHIParamRef TextureRHI,void
 	RHITHREAD_GLCOMMAND_EPILOGUE();
 }
 
-void FOpenGLDynamicRHI::RHIUnmapStagingSurface(FTextureRHIParamRef TextureRHI)
+void FOpenGLDynamicRHI::RHIUnmapStagingSurface(FRHITexture* TextureRHI)
 {
 	FRHICommandListImmediate& RHICmdList = FRHICommandListExecutor::GetImmediateCommandList();
 
@@ -971,7 +971,7 @@ void FOpenGLDynamicRHI::RHIUnmapStagingSurface(FTextureRHIParamRef TextureRHI)
 	RHITHREAD_GLCOMMAND_EPILOGUE();
 }
 
-void FOpenGLDynamicRHI::RHIReadSurfaceFloatData(FTextureRHIParamRef TextureRHI,FIntRect Rect,TArray<FFloat16Color>& OutData,ECubeFace CubeFace,int32 ArrayIndex,int32 MipIndex)
+void FOpenGLDynamicRHI::RHIReadSurfaceFloatData(FRHITexture* TextureRHI,FIntRect Rect,TArray<FFloat16Color>& OutData,ECubeFace CubeFace,int32 ArrayIndex,int32 MipIndex)
 {
 	FRHICommandListImmediate& RHICmdList = FRHICommandListExecutor::GetImmediateCommandList();
 
@@ -1047,7 +1047,7 @@ void FOpenGLDynamicRHI::RHIReadSurfaceFloatData(FTextureRHIParamRef TextureRHI,F
 	RHITHREAD_GLCOMMAND_EPILOGUE();
 }
 
-void FOpenGLDynamicRHI::RHIRead3DSurfaceFloatData(FTextureRHIParamRef TextureRHI,FIntRect Rect,FIntPoint ZMinMax,TArray<FFloat16Color>& OutData)
+void FOpenGLDynamicRHI::RHIRead3DSurfaceFloatData(FRHITexture* TextureRHI,FIntRect Rect,FIntPoint ZMinMax,TArray<FFloat16Color>& OutData)
 {
 	FRHICommandListImmediate& RHICmdList = FRHICommandListExecutor::GetImmediateCommandList();
 
diff --git a/Engine/Source/Runtime/OpenGLDrv/Private/OpenGLStructuredBuffer.cpp b/Engine/Source/Runtime/OpenGLDrv/Private/OpenGLStructuredBuffer.cpp
--- a/Engine/Source/Runtime/OpenGLDrv/Private/OpenGLStructuredBuffer.cpp
+++ b/Engine/Source/Runtime/OpenGLDrv/Private/OpenGLStructuredBuffer.cpp
@@ -25,14 +25,14 @@ FStructuredBufferRHIRef FOpenGLDynamicRHI::RHICreateStructuredBuffer(uint32 Stri
 	return StructuredBuffer.GetReference();
 }
 
-void* FOpenGLDynamicRHI::RHILockStructuredBuffer(FStructuredBufferRHIParamRef StructuredBufferRHI,uint32 Offset,uint32 Size,EResourceLockMode LockMode)
+void* FOpenGLDynamicRHI::RHILockStructuredBuffer(FRHIStructuredBuffer* StructuredBufferRHI,uint32 Offset,uint32 Size,EResourceLockMode LockMode)
 {
 	VERIFY_GL_SCOPE();
 	FOpenGLStructuredBuffer* StructuredBuffer = ResourceCast(StructuredBufferRHI);
 	return StructuredBuffer->Lock(Offset, Size, LockMode == RLM_ReadOnly, StructuredBuffer->IsDynamic());
 }
 
-void FOpenGLDynamicRHI::RHIUnlockStructuredBuffer(FStructuredBufferRHIParamRef StructuredBufferRHI)
+void FOpenGLDynamicRHI::RHIUnlockStructuredBuffer(FRHIStructuredBuffer* StructuredBufferRHI)
 {
 	VERIFY_GL_SCOPE();
 	FOpenGLStructuredBuffer* StructuredBuffer = ResourceCast(StructuredBufferRHI);
diff --git a/Engine/Source/Runtime/OpenGLDrv/Private/OpenGLTexture.cpp b/Engine/Source/Runtime/OpenGLDrv/Private/OpenGLTexture.cpp
--- a/Engine/Source/Runtime/OpenGLDrv/Private/OpenGLTexture.cpp
+++ b/Engine/Source/Runtime/OpenGLDrv/Private/OpenGLTexture.cpp
@@ -1597,7 +1597,7 @@ FTexture2DRHIRef FOpenGLDynamicRHI::RHIAsyncCreateTexture2D(uint32 SizeX, uint32
 	return FTexture2DRHIRef();
 }
 
-void FOpenGLDynamicRHI::RHICopySharedMips(FTexture2DRHIParamRef DestTexture2D,FTexture2DRHIParamRef SrcTexture2D)
+void FOpenGLDynamicRHI::RHICopySharedMips(FRHITexture2D* DestTexture2D, FRHITexture2D* SrcTexture2D)
 {
 	check(0);
 }
@@ -1853,11 +1853,11 @@ FTexture3DRHIRef FOpenGLDynamicRHI::RHICreateTexture3D(uint32 SizeX,uint32 SizeY
 	return Texture;
 }
 
-void FOpenGLDynamicRHI::RHIGetResourceInfo(FTextureRHIParamRef Ref, FRHIResourceInfo& OutInfo)
+void FOpenGLDynamicRHI::RHIGetResourceInfo(FRHITexture* Ref, FRHIResourceInfo& OutInfo)
 {
 }
 
-FShaderResourceViewRHIRef FOpenGLDynamicRHI::RHICreateShaderResourceView(FTextureRHIParamRef Texture, const FRHITextureSRVCreateInfo& CreateInfo)
+FShaderResourceViewRHIRef FOpenGLDynamicRHI::RHICreateShaderResourceView(FRHITexture* Texture, const FRHITextureSRVCreateInfo& CreateInfo)
 {
 	const uint32 MipLevel = CreateInfo.MipLevel;
 	const uint32 NumMipLevels = CreateInfo.NumMipLevels;
@@ -1914,7 +1914,7 @@ FShaderResourceViewRHIRef FOpenGLDynamicRHI::RHICreateShaderResourceView(FTextur
 				uint32 const Target = Texture2D->Target;
 				GLuint Resource = Texture2D->Resource;
 
-				FTexture2DRHIParamRef DepthStencilTex = nullptr;
+				FRHITexture2D* DepthStencilTex = nullptr;
 
 				// For stencil sampling we have to use a separate single channel texture to blit stencil data into
 #if PLATFORM_DESKTOP || PLATFORM_ANDROIDESDEFERRED
@@ -2056,7 +2056,7 @@ FShaderResourceViewRHIRef FOpenGLDynamicRHI::RHICreateShaderResourceView(FTextur
 }
 
 /** Generates mip maps for the surface. */
-void FOpenGLDynamicRHI::RHIGenerateMips(FTextureRHIParamRef SurfaceRHI)
+void FOpenGLDynamicRHI::RHIGenerateMips(FRHITexture* SurfaceRHI)
 {
 	VERIFY_GL_SCOPE();
 
@@ -2087,7 +2087,7 @@ void FOpenGLDynamicRHI::RHIGenerateMips(FTextureRHIParamRef SurfaceRHI)
  * @param	TextureRHI		- Texture we want to know the size of
  * @return					- Size in Bytes
  */
-uint32 FOpenGLDynamicRHI::RHIComputeMemorySize(FTextureRHIParamRef TextureRHI)
+uint32 FOpenGLDynamicRHI::RHIComputeMemorySize(FRHITexture* TextureRHI)
 {
 	if(!TextureRHI)
 	{
@@ -2103,7 +2103,7 @@ uint32 FOpenGLDynamicRHI::RHIComputeMemorySize(FTextureRHIParamRef TextureRHI)
 }
 
 
-static FTexture2DRHIRef CreateAsyncReallocate2DTextureTarget(FOpenGLDynamicRHI* OGLRHI, FTexture2DRHIParamRef Texture2DRHI, int32 NewMipCountIn, int32 NewSizeX, int32 NewSizeY)
+static FTexture2DRHIRef CreateAsyncReallocate2DTextureTarget(FOpenGLDynamicRHI* OGLRHI, FRHITexture2D* Texture2DRHI, int32 NewMipCountIn, int32 NewSizeX, int32 NewSizeY)
 {
 	FOpenGLTexture2D* Texture2D = FOpenGLDynamicRHI::ResourceCast(Texture2DRHI);
 	uint8 Format = Texture2D->GetFormat();
@@ -2121,7 +2121,7 @@ static FTexture2DRHIRef CreateAsyncReallocate2DTextureTarget(FOpenGLDynamicRHI*
 	return NewTexture2DRHI;
 }
 
-static void GLCopyAsyncTexture2D(FOpenGLDynamicRHI* OGLRHI, FRHITexture2D* NewTexture2DRHI, int32 NewSizeX, int32 NewSizeY, FTexture2DRHIParamRef SourceTexture2DRHI, FThreadSafeCounter* RequestStatus)
+static void GLCopyAsyncTexture2D(FOpenGLDynamicRHI* OGLRHI, FRHITexture2D* NewTexture2DRHI, int32 NewSizeX, int32 NewSizeY, FRHITexture2D* SourceTexture2DRHI, FThreadSafeCounter* RequestStatus)
 {
 	VERIFY_GL_SCOPE();
 
@@ -2197,7 +2197,7 @@ static void GLCopyAsyncTexture2D(FOpenGLDynamicRHI* OGLRHI, FRHITexture2D* NewTe
 	RequestStatus->Decrement();
 }
 
-FTexture2DRHIRef FOpenGLDynamicRHI::AsyncReallocateTexture2D_RenderThread(class FRHICommandListImmediate& RHICmdList, FTexture2DRHIParamRef Texture2DRHI, int32 NewMipCountIn, int32 NewSizeX, int32 NewSizeY, FThreadSafeCounter* RequestStatus)
+FTexture2DRHIRef FOpenGLDynamicRHI::AsyncReallocateTexture2D_RenderThread(class FRHICommandListImmediate& RHICmdList, FRHITexture2D* Texture2DRHI, int32 NewMipCountIn, int32 NewSizeX, int32 NewSizeY, FThreadSafeCounter* RequestStatus)
 {
 	if (RHICmdList.Bypass() || !IsRunningRHIInSeparateThread())
 	{
@@ -2220,7 +2220,7 @@ FTexture2DRHIRef FOpenGLDynamicRHI::AsyncReallocateTexture2D_RenderThread(class
 	}
 }
 
-FTexture2DRHIRef FOpenGLDynamicRHI::RHIAsyncReallocateTexture2D(FTexture2DRHIParamRef Texture2DRHI, int32 NewMipCountIn, int32 NewSizeX, int32 NewSizeY, FThreadSafeCounter* RequestStatus)
+FTexture2DRHIRef FOpenGLDynamicRHI::RHIAsyncReallocateTexture2D(FRHITexture2D* Texture2DRHI, int32 NewMipCountIn, int32 NewSizeX, int32 NewSizeY, FThreadSafeCounter* RequestStatus)
 {
 	FTexture2DRHIRef NewTexture2DRHI = CreateAsyncReallocate2DTextureTarget(this, Texture2DRHI, NewMipCountIn, NewSizeX, NewSizeY);
 	GLCopyAsyncTexture2D(this, NewTexture2DRHI, NewSizeX, NewSizeY, Texture2DRHI, RequestStatus);
@@ -2236,7 +2236,7 @@ FTexture2DRHIRef FOpenGLDynamicRHI::RHIAsyncReallocateTexture2D(FTexture2DRHIPar
  * @param Texture2D		- Texture to check the reallocation status for
  * @return				- Current reallocation status
  */
-ETextureReallocationStatus FOpenGLDynamicRHI::RHIFinalizeAsyncReallocateTexture2D( FTexture2DRHIParamRef Texture2D, bool bBlockUntilCompleted )
+ETextureReallocationStatus FOpenGLDynamicRHI::RHIFinalizeAsyncReallocateTexture2D(FRHITexture2D* Texture2D, bool bBlockUntilCompleted )
 {
 	return TexRealloc_Succeeded;
 }
@@ -2249,36 +2249,36 @@ ETextureReallocationStatus FOpenGLDynamicRHI::RHIFinalizeAsyncReallocateTexture2
  * @param bBlockUntilCompleted	If true, blocks until the cancellation is fully completed
  * @return						Reallocation status
  */
-ETextureReallocationStatus FOpenGLDynamicRHI::RHICancelAsyncReallocateTexture2D( FTexture2DRHIParamRef Texture2D, bool bBlockUntilCompleted )
+ETextureReallocationStatus FOpenGLDynamicRHI::RHICancelAsyncReallocateTexture2D(FRHITexture2D* Texture2D, bool bBlockUntilCompleted )
 {
 	return TexRealloc_Succeeded;
 }
 
-void* FOpenGLDynamicRHI::RHILockTexture2D(FTexture2DRHIParamRef TextureRHI,uint32 MipIndex,EResourceLockMode LockMode,uint32& DestStride,bool bLockWithinMiptail)
+void* FOpenGLDynamicRHI::RHILockTexture2D(FRHITexture2D* TextureRHI,uint32 MipIndex,EResourceLockMode LockMode,uint32& DestStride,bool bLockWithinMiptail)
 {
 	FOpenGLTexture2D* Texture = ResourceCast(TextureRHI);
 	return Texture->Lock(MipIndex,0,LockMode,DestStride);
 }
 
-void FOpenGLDynamicRHI::RHIUnlockTexture2D(FTexture2DRHIParamRef TextureRHI,uint32 MipIndex,bool bLockWithinMiptail)
+void FOpenGLDynamicRHI::RHIUnlockTexture2D(FRHITexture2D* TextureRHI,uint32 MipIndex,bool bLockWithinMiptail)
 {
 	FOpenGLTexture2D* Texture = ResourceCast(TextureRHI);
 	Texture->Unlock(MipIndex, 0);
 }
 
-void* FOpenGLDynamicRHI::RHILockTexture2DArray(FTexture2DArrayRHIParamRef TextureRHI,uint32 TextureIndex,uint32 MipIndex,EResourceLockMode LockMode,uint32& DestStride,bool bLockWithinMiptail)
+void* FOpenGLDynamicRHI::RHILockTexture2DArray(FRHITexture2DArray* TextureRHI,uint32 TextureIndex,uint32 MipIndex,EResourceLockMode LockMode,uint32& DestStride,bool bLockWithinMiptail)
 {
 	FOpenGLTexture2DArray* Texture = ResourceCast(TextureRHI);
 	return Texture->Lock(MipIndex,TextureIndex,LockMode,DestStride);
 }
 
-void FOpenGLDynamicRHI::RHIUnlockTexture2DArray(FTexture2DArrayRHIParamRef TextureRHI,uint32 TextureIndex,uint32 MipIndex,bool bLockWithinMiptail)
+void FOpenGLDynamicRHI::RHIUnlockTexture2DArray(FRHITexture2DArray* TextureRHI,uint32 TextureIndex,uint32 MipIndex,bool bLockWithinMiptail)
 {
 	FOpenGLTexture2DArray* Texture = ResourceCast(TextureRHI);
 	Texture->Unlock(MipIndex, TextureIndex);
 }
 
-void FOpenGLDynamicRHI::RHIUpdateTexture2D(FTexture2DRHIParamRef TextureRHI,uint32 MipIndex,const FUpdateTextureRegion2D& UpdateRegionIn,uint32 SourcePitch,const uint8* SourceDataIn)
+void FOpenGLDynamicRHI::RHIUpdateTexture2D(FRHITexture2D* TextureRHI,uint32 MipIndex,const FUpdateTextureRegion2D& UpdateRegionIn,uint32 SourcePitch,const uint8* SourceDataIn)
 {
 	FRHICommandListImmediate& RHICmdList = FRHICommandListExecutor::GetImmediateCommandList();
 	const FUpdateTextureRegion2D UpdateRegion = UpdateRegionIn;
@@ -2330,7 +2330,7 @@ void FOpenGLDynamicRHI::RHIUpdateTexture2D(FTexture2DRHIParamRef TextureRHI,uint
 	});
 }
 
-void FOpenGLDynamicRHI::RHIUpdateTexture3D(FTexture3DRHIParamRef TextureRHI,uint32 MipIndex,const FUpdateTextureRegion3D& UpdateRegion,uint32 SourceRowPitch,uint32 SourceDepthPitch,const uint8* SourceData)
+void FOpenGLDynamicRHI::RHIUpdateTexture3D(FRHITexture3D* TextureRHI, uint32 MipIndex, const FUpdateTextureRegion3D& UpdateRegion, uint32 SourceRowPitch, uint32 SourceDepthPitch, const uint8* SourceData)
 {
 	VERIFY_GL_SCOPE();
 	check( FOpenGL::SupportsTexture3D() );
@@ -2451,19 +2451,19 @@ FTextureCubeRHIRef FOpenGLDynamicRHI::RHICreateTextureCubeArray( uint32 Size, ui
 	return (FRHITextureCube*)CreateOpenGLTexture(Size, Size, true, true, false, Format, NumMips, 1, 6 * ArraySize, Flags, CreateInfo.ClearValueBinding);
 }
 
-void* FOpenGLDynamicRHI::RHILockTextureCubeFace(FTextureCubeRHIParamRef TextureCubeRHI,uint32 FaceIndex,uint32 ArrayIndex,uint32 MipIndex,EResourceLockMode LockMode,uint32& DestStride,bool bLockWithinMiptail)
+void* FOpenGLDynamicRHI::RHILockTextureCubeFace(FRHITextureCube* TextureCubeRHI,uint32 FaceIndex,uint32 ArrayIndex,uint32 MipIndex,EResourceLockMode LockMode,uint32& DestStride,bool bLockWithinMiptail)
 {
 	FOpenGLTextureCube* TextureCube = ResourceCast(TextureCubeRHI);
 	return TextureCube->Lock(MipIndex,FaceIndex + 6 * ArrayIndex,LockMode,DestStride);
 }
 
-void FOpenGLDynamicRHI::RHIUnlockTextureCubeFace(FTextureCubeRHIParamRef TextureCubeRHI,uint32 FaceIndex,uint32 ArrayIndex,uint32 MipIndex,bool bLockWithinMiptail)
+void FOpenGLDynamicRHI::RHIUnlockTextureCubeFace(FRHITextureCube* TextureCubeRHI,uint32 FaceIndex,uint32 ArrayIndex,uint32 MipIndex,bool bLockWithinMiptail)
 {
 	FOpenGLTextureCube* TextureCube = ResourceCast(TextureCubeRHI);
 	TextureCube->Unlock(MipIndex,FaceIndex + ArrayIndex * 6);
 }
 
-void FOpenGLDynamicRHI::RHIBindDebugLabelName(FTextureRHIParamRef TextureRHI, const TCHAR* Name)
+void FOpenGLDynamicRHI::RHIBindDebugLabelName(FRHITexture* TextureRHI, const TCHAR* Name)
 {
 	FRHICommandListImmediate& RHICmdList = FRHICommandListExecutor::GetImmediateCommandList();
 	if (ShouldRunGLRenderContextOpOnThisThread(RHICmdList))
@@ -2487,11 +2487,11 @@ void FOpenGLDynamicRHI::RHIBindDebugLabelName(FTextureRHIParamRef TextureRHI, co
 }
 
 
-void FOpenGLDynamicRHI::RHIVirtualTextureSetFirstMipInMemory(FTexture2DRHIParamRef TextureRHI, uint32 FirstMip)
+void FOpenGLDynamicRHI::RHIVirtualTextureSetFirstMipInMemory(FRHITexture2D* TextureRHI, uint32 FirstMip)
 {
 }
 
-void FOpenGLDynamicRHI::RHIVirtualTextureSetFirstMipVisible(FTexture2DRHIParamRef TextureRHI, uint32 FirstMip)
+void FOpenGLDynamicRHI::RHIVirtualTextureSetFirstMipVisible(FRHITexture2D* TextureRHI, uint32 FirstMip)
 {
 }
 
@@ -2506,7 +2506,7 @@ void FOpenGLTextureReference::SetReferencedTexture(FRHITexture* InTexture)
 	TexturePtr = GetOpenGLTextureFromRHITexture(InTexture);
 }
 
-void FOpenGLDynamicRHI::RHIUpdateTextureReference(FRHITextureReference* TextureRefRHI, FTextureRHIParamRef NewTextureRHI)
+void FOpenGLDynamicRHI::RHIUpdateTextureReference(FRHITextureReference* TextureRefRHI, FRHITexture* NewTextureRHI)
 {
 	auto* TextureRef = (FOpenGLTextureReference*)TextureRefRHI;
 	if (TextureRef)
@@ -2515,7 +2515,7 @@ void FOpenGLDynamicRHI::RHIUpdateTextureReference(FRHITextureReference* TextureR
 	}
 }
 
-void FOpenGLDynamicRHI::RHICopySubTextureRegion(FTexture2DRHIParamRef SourceTextureRHI, FTexture2DRHIParamRef DestinationTextureRHI, FBox2D SourceBox, FBox2D DestinationBox)
+void FOpenGLDynamicRHI::RHICopySubTextureRegion(FRHITexture2D* SourceTextureRHI, FRHITexture2D* DestinationTextureRHI, FBox2D SourceBox, FBox2D DestinationBox)
 {
 	VERIFY_GL_SCOPE();
 	FOpenGLTexture2D* SourceTexture = ResourceCast(SourceTextureRHI);
@@ -2550,7 +2550,7 @@ void FOpenGLDynamicRHI::RHICopySubTextureRegion(FTexture2DRHIParamRef SourceText
 	ContextState.Framebuffer = (GLuint)-1;
 }
 
-void FOpenGLDynamicRHI::RHICopyTexture(FTextureRHIParamRef SourceTextureRHI, FTextureRHIParamRef DestTextureRHI, const FRHICopyTextureInfo& CopyInfo)
+void FOpenGLDynamicRHI::RHICopyTexture(FRHITexture* SourceTextureRHI, FRHITexture* DestTextureRHI, const FRHICopyTextureInfo& CopyInfo)
 {
 	VERIFY_GL_SCOPE();
 	FOpenGLTextureBase* SourceTexture = GetOpenGLTextureFromRHITexture(SourceTextureRHI);
@@ -2693,7 +2693,7 @@ FTextureCubeRHIRef FOpenGLDynamicRHI::RHICreateTextureCubeFromResource(EPixelFor
 	return TextureCube;
 }
 
-void FOpenGLDynamicRHI::RHIAliasTextureResources(FTextureRHIParamRef DestRHITexture, FTextureRHIParamRef SrcRHITexture)
+void FOpenGLDynamicRHI::RHIAliasTextureResources(FRHITexture* DestRHITexture, FRHITexture* SrcRHITexture)
 {
 	FOpenGLTextureBase* DestTexture = GetOpenGLTextureFromRHITexture(DestRHITexture);
 	FOpenGLTextureBase* SrcTexture = GetOpenGLTextureFromRHITexture(SrcRHITexture);
@@ -2704,7 +2704,7 @@ void FOpenGLDynamicRHI::RHIAliasTextureResources(FTextureRHIParamRef DestRHIText
 	}
 }
 
-void* FOpenGLDynamicRHI::LockTexture2D_RenderThread(class FRHICommandListImmediate& RHICmdList, FTexture2DRHIParamRef Texture, uint32 MipIndex, EResourceLockMode LockMode, uint32& DestStride, bool bLockWithinMiptail, bool bNeedsDefaultRHIFlush)
+void* FOpenGLDynamicRHI::LockTexture2D_RenderThread(class FRHICommandListImmediate& RHICmdList, FRHITexture2D* Texture, uint32 MipIndex, EResourceLockMode LockMode, uint32& DestStride, bool bLockWithinMiptail, bool bNeedsDefaultRHIFlush)
 {
 	check(IsInRenderingThread());
 	static auto* CVarRHICmdBufferWriteLocks = IConsoleManager::Get().FindTConsoleVariableDataInt(TEXT("r.RHICmdBufferWriteLocks"));
@@ -2730,7 +2730,7 @@ void* FOpenGLDynamicRHI::LockTexture2D_RenderThread(class FRHICommandListImmedia
 	return Result;
 }
 
-void FOpenGLDynamicRHI::UnlockTexture2D_RenderThread(class FRHICommandListImmediate& RHICmdList, FTexture2DRHIParamRef Texture, uint32 MipIndex, bool bLockWithinMiptail, bool bNeedsDefaultRHIFlush)
+void FOpenGLDynamicRHI::UnlockTexture2D_RenderThread(class FRHICommandListImmediate& RHICmdList, FRHITexture2D* Texture, uint32 MipIndex, bool bLockWithinMiptail, bool bNeedsDefaultRHIFlush)
 {
 	check(IsInRenderingThread());
 	static auto* CVarRHICmdBufferWriteLocks = IConsoleManager::Get().FindTConsoleVariableDataInt(TEXT("r.RHICmdBufferWriteLocks"));
@@ -2759,7 +2759,7 @@ void FOpenGLDynamicRHI::UnlockTexture2D_RenderThread(class FRHICommandListImmedi
 	}
 }
 
-void* FOpenGLDynamicRHI::RHILockTextureCubeFace_RenderThread(class FRHICommandListImmediate& RHICmdList, FTextureCubeRHIParamRef Texture, uint32 FaceIndex, uint32 ArrayIndex, uint32 MipIndex, EResourceLockMode LockMode, uint32& DestStride, bool bLockWithinMiptail)
+void* FOpenGLDynamicRHI::RHILockTextureCubeFace_RenderThread(class FRHICommandListImmediate& RHICmdList, FRHITextureCube* Texture, uint32 FaceIndex, uint32 ArrayIndex, uint32 MipIndex, EResourceLockMode LockMode, uint32& DestStride, bool bLockWithinMiptail)
 {
 	check(IsInRenderingThread());
 	static auto* CVarRHICmdBufferWriteLocks = IConsoleManager::Get().FindTConsoleVariableDataInt(TEXT("r.RHICmdBufferWriteLocks"));
@@ -2784,7 +2784,7 @@ void* FOpenGLDynamicRHI::RHILockTextureCubeFace_RenderThread(class FRHICommandLi
 	return Result;
 }
 
-void FOpenGLDynamicRHI::RHIUnlockTextureCubeFace_RenderThread(class FRHICommandListImmediate& RHICmdList, FTextureCubeRHIParamRef Texture, uint32 FaceIndex, uint32 ArrayIndex, uint32 MipIndex, bool bLockWithinMiptail)
+void FOpenGLDynamicRHI::RHIUnlockTextureCubeFace_RenderThread(class FRHICommandListImmediate& RHICmdList, FRHITextureCube* Texture, uint32 FaceIndex, uint32 ArrayIndex, uint32 MipIndex, bool bLockWithinMiptail)
 {
 	check(IsInRenderingThread());
 	static auto* CVarRHICmdBufferWriteLocks = IConsoleManager::Get().FindTConsoleVariableDataInt(TEXT("r.RHICmdBufferWriteLocks"));
diff --git a/Engine/Source/Runtime/OpenGLDrv/Private/OpenGLUAV.cpp b/Engine/Source/Runtime/OpenGLDrv/Private/OpenGLUAV.cpp
--- a/Engine/Source/Runtime/OpenGLDrv/Private/OpenGLUAV.cpp
+++ b/Engine/Source/Runtime/OpenGLDrv/Private/OpenGLUAV.cpp
@@ -6,7 +6,7 @@
 #include "OpenGLDrv.h"
 #include "OpenGLDrvPrivate.h"
 
-FShaderResourceViewRHIRef FOpenGLDynamicRHI::RHICreateShaderResourceView(FVertexBufferRHIParamRef VertexBufferRHI, uint32 Stride, uint8 Format)
+FShaderResourceViewRHIRef FOpenGLDynamicRHI::RHICreateShaderResourceView(FRHIVertexBuffer* VertexBufferRHI, uint32 Stride, uint8 Format)
 {
 	FShaderResourceViewRHIRef Result = new FOpenGLShaderResourceViewProxy([=, OGLRHI = this](FRHIShaderResourceView* OwnerRHI)
 	{ 
@@ -47,7 +47,7 @@ FShaderResourceViewRHIRef FOpenGLDynamicRHI::RHICreateShaderResourceView(FVertex
 	return Result;
 }
 
-void FOpenGLDynamicRHI::RHIUpdateShaderResourceView(FRHIShaderResourceView* SRV, FVertexBufferRHIParamRef VertexBuffer, uint32 Stride, uint8 Format)
+void FOpenGLDynamicRHI::RHIUpdateShaderResourceView(FRHIShaderResourceView* SRV, FRHIVertexBuffer* VertexBuffer, uint32 Stride, uint8 Format)
 {
 	if (!FOpenGL::SupportsResourceView())
 	{
@@ -77,7 +77,7 @@ void FOpenGLDynamicRHI::RHIUpdateShaderResourceView(FRHIShaderResourceView* SRV,
 	}
 }
 
-FShaderResourceViewRHIRef FOpenGLDynamicRHI::RHICreateShaderResourceView(FIndexBufferRHIParamRef BufferRHI)
+FShaderResourceViewRHIRef FOpenGLDynamicRHI::RHICreateShaderResourceView(FRHIIndexBuffer* BufferRHI)
 {
 	return new FOpenGLShaderResourceViewProxy([=](FRHIShaderResourceView* OwnerRHI)
 	{
@@ -107,7 +107,7 @@ FOpenGLShaderResourceView::~FOpenGLShaderResourceView()
 	}
 }
 
-FUnorderedAccessViewRHIRef FOpenGLDynamicRHI::RHICreateUnorderedAccessView(FStructuredBufferRHIParamRef StructuredBufferRHI, bool bUseUAVCounter, bool bAppendBuffer)
+FUnorderedAccessViewRHIRef FOpenGLDynamicRHI::RHICreateUnorderedAccessView(FRHIStructuredBuffer* StructuredBufferRHI, bool bUseUAVCounter, bool bAppendBuffer)
 {
 	FOpenGLStructuredBuffer* StructuredBuffer = ResourceCast(StructuredBufferRHI);
 	// emulate structured buffer of specific size as typed buffer
@@ -126,7 +126,7 @@ FUnorderedAccessViewRHIRef FOpenGLDynamicRHI::RHICreateUnorderedAccessView(FStru
 	return new FOpenGLUnorderedAccessView();
 }
 
-FUnorderedAccessViewRHIRef FOpenGLDynamicRHI::RHICreateUnorderedAccessView(FTextureRHIParamRef TextureRHI, uint32 MipLevel)
+FUnorderedAccessViewRHIRef FOpenGLDynamicRHI::RHICreateUnorderedAccessView(FRHITexture* TextureRHI, uint32 MipLevel)
 {
 	FOpenGLTexture* Texture = ResourceCast(TextureRHI);
 	check(Texture->GetFlags() & TexCreate_UAV);
@@ -134,7 +134,7 @@ FUnorderedAccessViewRHIRef FOpenGLDynamicRHI::RHICreateUnorderedAccessView(FText
 }
 
 
-FOpenGLTextureUnorderedAccessView::FOpenGLTextureUnorderedAccessView( FTextureRHIParamRef InTextureRHI):
+FOpenGLTextureUnorderedAccessView::FOpenGLTextureUnorderedAccessView(FRHITexture* InTextureRHI):
 	TextureRHI(InTextureRHI)
 {
 	VERIFY_GL_SCOPE();
@@ -147,7 +147,7 @@ FOpenGLTextureUnorderedAccessView::FOpenGLTextureUnorderedAccessView( FTextureRH
 }
 
 
-FOpenGLVertexBufferUnorderedAccessView::FOpenGLVertexBufferUnorderedAccessView(	FOpenGLDynamicRHI* InOpenGLRHI, FVertexBufferRHIParamRef InVertexBufferRHI, uint8 Format):
+FOpenGLVertexBufferUnorderedAccessView::FOpenGLVertexBufferUnorderedAccessView(	FOpenGLDynamicRHI* InOpenGLRHI, FRHIVertexBuffer* InVertexBufferRHI, uint8 Format):
 	VertexBufferRHI(InVertexBufferRHI),
 	OpenGLRHI(InOpenGLRHI)
 {
@@ -189,19 +189,19 @@ FOpenGLVertexBufferUnorderedAccessView::~FOpenGLVertexBufferUnorderedAccessView(
 }
 
 
-FUnorderedAccessViewRHIRef FOpenGLDynamicRHI::RHICreateUnorderedAccessView(FVertexBufferRHIParamRef VertexBufferRHI,uint8 Format)
+FUnorderedAccessViewRHIRef FOpenGLDynamicRHI::RHICreateUnorderedAccessView(FRHIVertexBuffer* VertexBufferRHI,uint8 Format)
 {
 	FOpenGLVertexBuffer* VertexBuffer = ResourceCast(VertexBufferRHI);
 	return new FOpenGLVertexBufferUnorderedAccessView(this, VertexBufferRHI, Format);
 }
 
-FUnorderedAccessViewRHIRef FOpenGLDynamicRHI::RHICreateUnorderedAccessView(FIndexBufferRHIParamRef IndexBufferRHI, uint8 Format)
+FUnorderedAccessViewRHIRef FOpenGLDynamicRHI::RHICreateUnorderedAccessView(FRHIIndexBuffer* IndexBufferRHI, uint8 Format)
 {
 	checkf(0, TEXT("Not implemented!"));
 	return nullptr;
 }
 
-FShaderResourceViewRHIRef FOpenGLDynamicRHI::RHICreateShaderResourceView(FStructuredBufferRHIParamRef StructuredBufferRHI)
+FShaderResourceViewRHIRef FOpenGLDynamicRHI::RHICreateShaderResourceView(FRHIStructuredBuffer* StructuredBufferRHI)
 {
 	return new FOpenGLShaderResourceViewProxy([=](FRHIShaderResourceView* OwnerRHI)
 	{
@@ -244,7 +244,7 @@ void FOpenGLDynamicRHI::RHIClearTinyUAV(FRHIUnorderedAccessView* UnorderedAccess
 #endif
 }
 
-FOpenGLStructuredBufferUnorderedAccessView::FOpenGLStructuredBufferUnorderedAccessView(FOpenGLDynamicRHI* InOpenGLRHI, FStructuredBufferRHIParamRef InStructuredBufferRHI, uint8 InFormat)
+FOpenGLStructuredBufferUnorderedAccessView::FOpenGLStructuredBufferUnorderedAccessView(FOpenGLDynamicRHI* InOpenGLRHI, FRHIStructuredBuffer* InStructuredBufferRHI, uint8 InFormat)
 	: StructuredBufferRHI(InStructuredBufferRHI)
 	, OpenGLRHI(InOpenGLRHI)
 {
diff --git a/Engine/Source/Runtime/OpenGLDrv/Private/OpenGLUniformBuffer.cpp b/Engine/Source/Runtime/OpenGLDrv/Private/OpenGLUniformBuffer.cpp
--- a/Engine/Source/Runtime/OpenGLDrv/Private/OpenGLUniformBuffer.cpp
+++ b/Engine/Source/Runtime/OpenGLDrv/Private/OpenGLUniformBuffer.cpp
@@ -720,7 +720,7 @@ void UpdateUniformBufferContents(FOpenGLUniformBuffer* UniformBuffer, const void
 	}
 }
 
-void FOpenGLDynamicRHI::RHIUpdateUniformBuffer(FUniformBufferRHIParamRef UniformBufferRHI, const void* Contents)
+void FOpenGLDynamicRHI::RHIUpdateUniformBuffer(FRHIUniformBuffer* UniformBufferRHI, const void* Contents)
 {
 	FOpenGLUniformBuffer* UniformBuffer = ResourceCast(UniformBufferRHI);
 
diff --git a/Engine/Source/Runtime/OpenGLDrv/Private/OpenGLVertexBuffer.cpp b/Engine/Source/Runtime/OpenGLDrv/Private/OpenGLVertexBuffer.cpp
--- a/Engine/Source/Runtime/OpenGLDrv/Private/OpenGLVertexBuffer.cpp
+++ b/Engine/Source/Runtime/OpenGLDrv/Private/OpenGLVertexBuffer.cpp
@@ -192,7 +192,7 @@ FVertexBufferRHIRef FOpenGLDynamicRHI::RHICreateVertexBuffer(uint32 Size, uint32
 	return VertexBuffer.GetReference();
 }
 
-void* FOpenGLDynamicRHI::RHILockVertexBuffer(FVertexBufferRHIParamRef VertexBufferRHI,uint32 Offset,uint32 Size,EResourceLockMode LockMode)
+void* FOpenGLDynamicRHI::RHILockVertexBuffer(FRHIVertexBuffer* VertexBufferRHI,uint32 Offset,uint32 Size,EResourceLockMode LockMode)
 {
 	check(Size > 0);
 	FRHICommandListImmediate& RHICmdList = FRHICommandListExecutor::GetImmediateCommandList();
@@ -222,7 +222,7 @@ void* FOpenGLDynamicRHI::RHILockVertexBuffer(FVertexBufferRHIParamRef VertexBuff
 	RHITHREAD_GLCOMMAND_EPILOGUE_RETURN(void*);
 }
 
-void FOpenGLDynamicRHI::RHIUnlockVertexBuffer(FVertexBufferRHIParamRef VertexBufferRHI)
+void FOpenGLDynamicRHI::RHIUnlockVertexBuffer(FRHIVertexBuffer* VertexBufferRHI)
 {
 	FRHICommandListImmediate& RHICmdList = FRHICommandListExecutor::GetImmediateCommandList();
 	RHITHREAD_GLCOMMAND_PROLOGUE();
@@ -238,7 +238,7 @@ void FOpenGLDynamicRHI::RHIUnlockVertexBuffer(FVertexBufferRHIParamRef VertexBuf
 	RHITHREAD_GLCOMMAND_EPILOGUE();
 }
 
-void FOpenGLDynamicRHI::RHICopyVertexBuffer(FVertexBufferRHIParamRef SourceBufferRHI,FVertexBufferRHIParamRef DestBufferRHI)
+void FOpenGLDynamicRHI::RHICopyVertexBuffer(FRHIVertexBuffer* SourceBufferRHI, FRHIVertexBuffer* DestBufferRHI)
 {
 	VERIFY_GL_SCOPE();
 	check( FOpenGL::SupportsCopyBuffer() );
diff --git a/Engine/Source/Runtime/OpenGLDrv/Private/OpenGLViewport.cpp b/Engine/Source/Runtime/OpenGLDrv/Private/OpenGLViewport.cpp
--- a/Engine/Source/Runtime/OpenGLDrv/Private/OpenGLViewport.cpp
+++ b/Engine/Source/Runtime/OpenGLDrv/Private/OpenGLViewport.cpp
@@ -89,7 +89,7 @@ void FOpenGLDynamicRHI::RHITick( float DeltaTime )
  *	Viewport functions.
  *=============================================================================*/
 
-void FOpenGLDynamicRHI::RHIBeginDrawingViewport(FRHIViewport* ViewportRHI, FTextureRHIParamRef RenderTarget)
+void FOpenGLDynamicRHI::RHIBeginDrawingViewport(FRHIViewport* ViewportRHI, FRHITexture* RenderTarget)
 {
 	VERIFY_GL_SCOPE();
 
