Back out nullrhi skipping RHI resources change
#rnx

diff --git a/Engine/Source/Runtime/NullDrv/Private/NullRHI.cpp b/Engine/Source/Runtime/NullDrv/Private/NullRHI.cpp
--- a/Engine/Source/Runtime/NullDrv/Private/NullRHI.cpp
+++ b/Engine/Source/Runtime/NullDrv/Private/NullRHI.cpp
@@ -35,9 +35,9 @@ void FNullDynamicRHI::Init()
 
 	check(!GIsRHIInitialized);
 
-	// Create render resources if app supports rendering
+	// do not do this at least on dedicated server; clients with -NullRHI may need additional consideration
 #if !WITH_EDITOR	
-	if (FApp::CanEverRender())
+	if (!IsRunningDedicatedServer())
 #endif
 	{
 		// Notify all initialized FRenderResources that there's a valid RHI device to create their RHI resources for now.
