Tessellation Removal: Removing Tessellation from NullRHI

diff --git a/Engine/Source/Runtime/NullDrv/Public/NullRHI.h b/Engine/Source/Runtime/NullDrv/Public/NullRHI.h
--- a/Engine/Source/Runtime/NullDrv/Public/NullRHI.h
+++ b/Engine/Source/Runtime/NullDrv/Public/NullRHI.h
@@ -52,16 +52,6 @@ public:
 		return new FRHIVertexShader(); 
 	}
 
-	virtual FHullShaderRHIRef RHICreateHullShader(TArrayView<const uint8> Code, const FSHAHash& Hash) final override
-	{ 
-		return new FRHIHullShader(); 
-	}
-
-	virtual FDomainShaderRHIRef RHICreateDomainShader(TArrayView<const uint8> Code, const FSHAHash& Hash) final override
-	{ 
-		return new FRHIDomainShader(); 
-	}
-
 	virtual FGeometryShaderRHIRef RHICreateGeometryShader(TArrayView<const uint8> Code, const FSHAHash& Hash) final override
 	{ 
 		return new FRHIGeometryShader(); 
