Fixing NullRHI on platforms with static shader platforms/feature levels; GMaxRHIFeatureLevel always needs to be set based on GMaxRHIShaderPlatform.
Adding SM6->SP_PCD3D_SM6 mapping to NullRHI on PC.

diff --git a/Engine/Source/Runtime/NullDrv/Private/NullRHI.cpp b/Engine/Source/Runtime/NullDrv/Private/NullRHI.cpp
--- a/Engine/Source/Runtime/NullDrv/Private/NullRHI.cpp
+++ b/Engine/Source/Runtime/NullDrv/Private/NullRHI.cpp
@@ -19,11 +19,14 @@ FNullDynamicRHI::FNullDynamicRHI()
 
 void FNullDynamicRHI::Init()
 {
+	GMaxRHIFeatureLevel = GetMaxSupportedFeatureLevel(GMaxRHIShaderPlatform);
+
 #if PLATFORM_WINDOWS
 	GShaderPlatformForFeatureLevel[ERHIFeatureLevel::ES2_REMOVED] = SP_NumPlatforms;
 	GShaderPlatformForFeatureLevel[ERHIFeatureLevel::ES3_1] = SP_PCD3D_ES3_1;
 	GShaderPlatformForFeatureLevel[ERHIFeatureLevel::SM4_REMOVED] = SP_NumPlatforms;
 	GShaderPlatformForFeatureLevel[ERHIFeatureLevel::SM5] = SP_PCD3D_SM5;
+	GShaderPlatformForFeatureLevel[ERHIFeatureLevel::SM6] = SP_PCD3D_SM6;
 #elif PLATFORM_MAC
     GShaderPlatformForFeatureLevel[ERHIFeatureLevel::ES2_REMOVED] = SP_NumPlatforms;
     GShaderPlatformForFeatureLevel[ERHIFeatureLevel::ES3_1] = SP_METAL_MACES3_1;
@@ -31,7 +34,7 @@ void FNullDynamicRHI::Init()
     GShaderPlatformForFeatureLevel[ERHIFeatureLevel::SM5] = SP_METAL_SM5;
     GShaderPlatformForFeatureLevel[ERHIFeatureLevel::SM6] = SP_METAL_SM6;
 #else
-	GShaderPlatformForFeatureLevel[GetMaxSupportedFeatureLevel(GMaxRHIShaderPlatform)] = GMaxRHIShaderPlatform;
+	GShaderPlatformForFeatureLevel[GMaxRHIFeatureLevel] = GMaxRHIShaderPlatform;
 #endif
 	
 	GRHIVendorId = 1;
