Fix NullRHI not setting up shader platforms on Linux.

diff --git a/Engine/Source/Runtime/NullDrv/Private/NullRHI.cpp b/Engine/Source/Runtime/NullDrv/Private/NullRHI.cpp
--- a/Engine/Source/Runtime/NullDrv/Private/NullRHI.cpp
+++ b/Engine/Source/Runtime/NullDrv/Private/NullRHI.cpp
@@ -33,6 +33,12 @@ void FNullDynamicRHI::Init()
     GShaderPlatformForFeatureLevel[ERHIFeatureLevel::SM4_REMOVED] = SP_NumPlatforms;
     GShaderPlatformForFeatureLevel[ERHIFeatureLevel::SM5] = SP_METAL_SM5;
     GShaderPlatformForFeatureLevel[ERHIFeatureLevel::SM6] = SP_METAL_SM6;
+#elif PLATFORM_LINUX // (see FVulkanGenericPlatform::SetupFeatureLevels)
+	GShaderPlatformForFeatureLevel[ERHIFeatureLevel::ES2_REMOVED] = SP_NumPlatforms;
+	GShaderPlatformForFeatureLevel[ERHIFeatureLevel::ES3_1] = SP_VULKAN_PCES3_1;
+	GShaderPlatformForFeatureLevel[ERHIFeatureLevel::SM4_REMOVED] = SP_NumPlatforms;
+	GShaderPlatformForFeatureLevel[ERHIFeatureLevel::SM5] = SP_VULKAN_SM5;
+	GShaderPlatformForFeatureLevel[ERHIFeatureLevel::SM6] = SP_VULKAN_SM6;
 #else
 	GShaderPlatformForFeatureLevel[GMaxRHIFeatureLevel] = GMaxRHIShaderPlatform;
 #endif
