Make NullRHI emit the GPU profiler frame boundary events
- The stats system relies on the RHI to submit frame boundaries in order to process the end-of-pipe stats timeline.
- Even though NullRHI does no rendering, we still need to emit this event to keep the stats system working.

diff --git a/Engine/Source/Runtime/NullDrv/Private/NullRHI.cpp b/Engine/Source/Runtime/NullDrv/Private/NullRHI.cpp
--- a/Engine/Source/Runtime/NullDrv/Private/NullRHI.cpp
+++ b/Engine/Source/Runtime/NullDrv/Private/NullRHI.cpp
@@ -7,6 +7,11 @@
 #include "RenderUtils.h"
 #include "RHICommandList.h"
 #include "DataDrivenShaderPlatformInfo.h"
+#include "GPUProfiler.h"
+
+#if RHI_NEW_GPU_PROFILER
+constexpr UE::RHI::GPUProfiler::FQueue NullRHIGPUProfilerQueue(UE::RHI::GPUProfiler::FQueue::EType::Graphics, 0, 0);
+#endif
 
 FNullDynamicRHI::FNullDynamicRHI()
 {
@@ -56,6 +61,11 @@ void FNullDynamicRHI::Init()
 	}
 
 	GIsRHIInitialized = true;
+
+#if RHI_NEW_GPU_PROFILER
+	// Register a single dummy graphics GPU queue. This is required so that we can still submit the frame boundary events in RHIEndFrame.
+	UE::RHI::GPUProfiler::InitializeQueues(MakeConstArrayView(&NullRHIGPUProfilerQueue, 1));
+#endif
 }
 
 
@@ -95,6 +105,25 @@ void* FNullDynamicRHI::GetStaticTextureBuffer(int32 SizeX, int32 SizeY, EPixelFo
 	return GetStaticBuffer(Size);
 }
 
+void FNullDynamicRHI::RHIEndFrame(const FRHIEndFrameArgs& Args)
+{
+	// The RHI must publish the frame boundary events for engine stats to work correctly.
+
+#if RHI_NEW_GPU_PROFILER
+	UE::RHI::GPUProfiler::FEventStream EventStream(NullRHIGPUProfilerQueue);
+	EventStream.Emplace<UE::RHI::GPUProfiler::FEvent::FFrameBoundary>(FPlatformTime::Cycles64(), Args.FrameNumber
+	#if WITH_RHI_BREADCRUMBS
+		, Args.GPUBreadcrumbs[ERHIPipeline::Graphics]
+	#endif
+	#if STATS
+		, Args.StatsFrame
+	#endif
+	);
+
+	UE::RHI::GPUProfiler::ProcessEvents(MakeArrayView(&EventStream, 1));
+#endif
+}
+
 /** Value between 0-100 that determines the percentage of the vertical scan that is allowed to pass while still allowing us to swap when VSYNC'ed.
 This is used to get the same behavior as the old *_OR_IMMEDIATE present modes. */
 uint32 GPresentImmediateThreshold = 100;
diff --git a/Engine/Source/Runtime/NullDrv/Public/NullRHI.h b/Engine/Source/Runtime/NullDrv/Public/NullRHI.h
--- a/Engine/Source/Runtime/NullDrv/Public/NullRHI.h
+++ b/Engine/Source/Runtime/NullDrv/Public/NullRHI.h
@@ -383,10 +383,7 @@ public:
 		return new FNullTexture(Desc);
 	}
 
-	virtual void RHIEndFrame(const FRHIEndFrameArgs& Args) final override
-	{
-
-	}
+	virtual void RHIEndFrame(const FRHIEndFrameArgs& Args) final override;
 
 	virtual void RHIAliasTextureResources(FTextureRHIRef& DestTexture, FTextureRHIRef& SrcTexture) final override
 	{
